Make MusicEnvironment module an Engine/Plugins/Runtime module instead of an "always on", Engine/Source/Runtime module.

diff --git a/Engine/Plugins/Runtime/Harmonix/Harmonix.uplugin b/Engine/Plugins/Runtime/Harmonix/Harmonix.uplugin
--- a/Engine/Plugins/Runtime/Harmonix/Harmonix.uplugin
+++ b/Engine/Plugins/Runtime/Harmonix/Harmonix.uplugin
@@ -73,6 +73,10 @@
 		}
 	],
 	"Plugins": [
+		{
+			"Name": "MusicEnvironment",
+			"Enabled": true
+		},
 		{
 			"Name": "Metasound",
 			"Enabled": true
diff --git a/Engine/Plugins/Runtime/Harmonix/Source/HarmonixMetasound/Private/HarmonixMetasound/MusicEnvironmentSupport/HarmonixMusicEnvironmentMetronome.cpp b/Engine/Plugins/Runtime/Harmonix/Source/HarmonixMetasound/Private/HarmonixMetasound/MusicEnvironmentSupport/HarmonixMusicEnvironmentMetronome.cpp
--- a/Engine/Plugins/Runtime/Harmonix/Source/HarmonixMetasound/Private/HarmonixMetasound/MusicEnvironmentSupport/HarmonixMusicEnvironmentMetronome.cpp
+++ b/Engine/Plugins/Runtime/Harmonix/Source/HarmonixMetasound/Private/HarmonixMetasound/MusicEnvironmentSupport/HarmonixMusicEnvironmentMetronome.cpp
@@ -37,7 +37,7 @@ void UHarmonixMusicEnvironmentMetronome::Tick(float DeltaSecs)
 
 void UHarmonixMusicEnvironmentMetronome::Start(double FromSeconds)
 {
-	UE_LOG(LogMIDI, Log, TEXT("UHarmonixMusicEnvironmentMetronome: Starting Metronome at %f seconds"), FromSeconds);
+	UE_LOG(LogMIDI, Verbose, TEXT("UHarmonixMusicEnvironmentMetronome: Starting Metronome at %f seconds"), FromSeconds);
 	FAudioParameter SeekTargetMs(TEXT("SeekTarget"), static_cast<float>(FromSeconds * 1000.0f));
 	AudioComponent->Activate();
 	AudioComponent->SetParameter(MoveTemp(SeekTargetMs));
@@ -51,7 +51,7 @@ void UHarmonixMusicEnvironmentMetronome::Seek(double ToSeconds)
 
 void UHarmonixMusicEnvironmentMetronome::Stop()
 {
-	UE_LOG(LogMIDI, Log, TEXT("UHarmonixMusicEnvironmentMetronome: Stopping Metronome"));
+	UE_LOG(LogMIDI, Verbose, TEXT("UHarmonixMusicEnvironmentMetronome: Stopping Metronome"));
 	AudioComponent->SetTriggerParameter(TEXT("StopMetronome"));
 	MusicClockComponent->Stop();
 	AudioComponent->Deactivate();
diff --git a/Engine/Plugins/Runtime/MusicEnvironment/MusicEnvironment.uplugin b/Engine/Plugins/Runtime/MusicEnvironment/MusicEnvironment.uplugin
new file mode 100644
--- /dev/null
+++ b/Engine/Plugins/Runtime/MusicEnvironment/MusicEnvironment.uplugin
@@ -0,0 +1,27 @@
+{
+	"FileVersion": 3,
+	"Version": 1,
+	"VersionName": "1.0",
+	"FriendlyName": "Music Environment",
+	"Description": "A Project-Wide source of musical information (musically synchronized clocks, events, etc.)",
+	"Category": "Audio",
+	"CreatedBy": "Epic Games, Inc. - Harmonix GenTech",
+	"CreatedByURL": "",
+	"DocsURL": "",
+	"MarketplaceURL": "",
+	"SupportURL": "",
+	"EnabledByDefault": false,
+	"CanContainContent": false,
+	"IsBetaVersion": true,
+	"IsExperimentalVersion": true,
+	"Installed": false,
+	"Modules": [
+		{
+			"Name": "MusicEnvironment",
+			"Type": "Runtime",
+			"LoadingPhase": "PreDefault"
+		}
+	],
+	"Plugins": [
+	]
+}
\ No newline at end of file
diff --git a/Engine/Plugins/Runtime/MusicEnvironment/Source/MusicEnvironment/MusicEnvironment.Build.cs b/Engine/Plugins/Runtime/MusicEnvironment/Source/MusicEnvironment/MusicEnvironment.Build.cs
new file mode 100644
--- /dev/null
+++ b/Engine/Plugins/Runtime/MusicEnvironment/Source/MusicEnvironment/MusicEnvironment.Build.cs
@@ -0,0 +1,25 @@
+// Copyright Epic Games, Inc. All Rights Reserved.
+
+using UnrealBuildTool;
+
+public class MusicEnvironment : ModuleRules
+{
+	public MusicEnvironment(ReadOnlyTargetRules Target) : base(Target)
+	{
+        PrivateIncludePathModuleNames.AddRange(
+            new string[] {
+			}
+        );
+
+        PublicDependencyModuleNames.AddRange(
+			new string[] {
+				"Core",
+				"CoreUObject",
+                "Engine",
+				"TimeManagement",
+				"GameplayTags",
+				"MovieScene",
+			}
+		);
+	}
+}
diff --git a/Engine/Plugins/Runtime/MusicEnvironment/Source/MusicEnvironment/Private/FrameBasedMusicMap.cpp b/Engine/Plugins/Runtime/MusicEnvironment/Source/MusicEnvironment/Private/FrameBasedMusicMap.cpp
new file mode 100644
--- /dev/null
+++ b/Engine/Plugins/Runtime/MusicEnvironment/Source/MusicEnvironment/Private/FrameBasedMusicMap.cpp
@@ -0,0 +1,569 @@
+// Copyright Epic Games, Inc. All Rights Reserved.
+
+#include "FrameBasedMusicMap.h"
+#include "Algo/FindLast.h"
+
+DEFINE_LOG_CATEGORY_STATIC(LogMusicMap, Log, All);
+#define LOCTEXT_NAMESPACE "FrameBasedMusicMap"
+
+namespace FrameBasedMusicMap
+{
+	FText GetQuantizationName(int32 Quantization)
+	{
+		switch (Quantization)
+		{
+		case ETripletQuantizationTicks::ThirtysecondNoteTriplet:	return LOCTEXT("ThirtySecondNoteTriplet", "1/32 Note Triplet");
+		case ETripletQuantizationTicks::SixteenthNoteTriplet:		return LOCTEXT("SixteenthNoteTriplet", "1/16 Note Triplet");
+		case ETripletQuantizationTicks::EighthNoteTriplet:			return LOCTEXT("EighthNoteTriplet", "1/8 Note Triplet");
+		case ETripletQuantizationTicks::QuarterNoteTriplet:			return LOCTEXT("QuarterNoteTriplet", "1/4 Note Triplet");
+		case ETripletQuantizationTicks::HalfNoteTriplet:			return LOCTEXT("HalfNoteTriplet", "1/2 Note Triplet");
+		case ETripletQuantizationTicks::WholeNoteTriplet:			return LOCTEXT("WholeNoteTriplet", "Whole Note Triplet");
+		
+		case EWholeQuantizationTicks::ThirtysecondNote:	return LOCTEXT("ThirtySecondNote", "1/32 Note");
+		case EWholeQuantizationTicks::SixteenthNote:	return LOCTEXT("SixteenthNote", "1/16 Note");
+		case EWholeQuantizationTicks::EighthNote:		return LOCTEXT("EighthNote", "1/8 Note");
+		case EWholeQuantizationTicks::QuarterNote:		return LOCTEXT("QuarterNote", "1/4 Note");
+		case EWholeQuantizationTicks::HalfNote:			return LOCTEXT("HalfNote", "1/2 Note");
+		case EWholeQuantizationTicks::WholeNote:		return LOCTEXT("WholeNote", "Whole Note");
+		case EWholeQuantizationTicks::Bar:				return LOCTEXT("Bar", "Bar");
+		case EWholeQuantizationTicks::Beat:				return LOCTEXT("Beat", "Beat");
+
+
+		case EDottedQuantizationTicks::DottedThirtysecondNote:	return LOCTEXT("DottedThirtySecondNote", "Dotted 1/32 Note");
+		case EDottedQuantizationTicks::DottedSixteenthNote:		return LOCTEXT("DottedSixteenthNote", "Dotted 1/16 Note");
+		case EDottedQuantizationTicks::DottedEighthNote:		return LOCTEXT("DottedEighthNote", "Dotted 1/8 Note");
+		case EDottedQuantizationTicks::DottedQuarterNote:		return LOCTEXT("DottedQuarterNote", "Dotted 1/4 Note");
+		case EDottedQuantizationTicks::DottedHalfNote:			return LOCTEXT("DottedHalfNote", "Dotted Half Note");
+		case EDottedQuantizationTicks::DottedWholeNote:			return LOCTEXT("DottedWholeNote", "Dotted Whole Note");
+		}
+		return LOCTEXT("Unknown", "Unrecognized Qunatization Factor");
+	}
+
+	TArray<TSharedPtr<int32>> QuantizationOptionsSource = {
+		MakeShared<int32>(EWholeQuantizationTicks::Bar),
+		MakeShared<int32>(EWholeQuantizationTicks::Beat),
+		MakeShared<int32>(EWholeQuantizationTicks::WholeNote),
+		MakeShared<int32>(EWholeQuantizationTicks::HalfNote),
+		MakeShared<int32>(EWholeQuantizationTicks::QuarterNote),
+		MakeShared<int32>(EWholeQuantizationTicks::EighthNote),
+		MakeShared<int32>(EWholeQuantizationTicks::SixteenthNote),
+		MakeShared<int32>(EWholeQuantizationTicks::ThirtysecondNote),
+		MakeShared<int32>(ETripletQuantizationTicks::WholeNoteTriplet),
+		MakeShared<int32>(ETripletQuantizationTicks::HalfNoteTriplet),
+		MakeShared<int32>(ETripletQuantizationTicks::QuarterNoteTriplet),
+		MakeShared<int32>(ETripletQuantizationTicks::EighthNoteTriplet),
+		MakeShared<int32>(ETripletQuantizationTicks::SixteenthNoteTriplet),
+		MakeShared<int32>(ETripletQuantizationTicks::ThirtysecondNoteTriplet),
+		MakeShared<int32>(EDottedQuantizationTicks::DottedWholeNote),
+		MakeShared<int32>(EDottedQuantizationTicks::DottedHalfNote),
+		MakeShared<int32>(EDottedQuantizationTicks::DottedQuarterNote),
+		MakeShared<int32>(EDottedQuantizationTicks::DottedEighthNote),
+		MakeShared<int32>(EDottedQuantizationTicks::DottedSixteenthNote),
+		MakeShared<int32>(EDottedQuantizationTicks::DottedThirtysecondNote)
+	};
+}
+
+void UFrameBasedMusicMap::SetFrameResolution(const FFrameRate& NewFrameResolution)
+{
+	check(NewFrameResolution.IsValid());
+
+	if (NewFrameResolution == FrameResolution)
+	{
+		return;
+	}
+
+	if (FrameResolution.IsValid() && !IsEmpty() && !IsSimple())
+	{
+		// must be changing sequence's resolution...
+		for (FFrameBasedTimeSignaturePoint& Bar : BarMap)
+		{
+			Bar.OnFrame = FFrameRate::TransformTime(FFrameTime::FromDecimal(Bar.OnFrame), FrameResolution, NewFrameResolution).AsDecimal();
+		}
+	}
+
+	FrameResolution = NewFrameResolution;
+}
+
+void UFrameBasedMusicMap::Init(float InBpm, const FFrameBasedTimeSignature& InTimeSignature)
+{
+	check(FrameResolution.IsValid());
+	Clear();
+	TempoMap.Add(FFrameBasedTempoPoint(InBpm, 0, 0.0f));
+	BarMap.Add(FFrameBasedTimeSignaturePoint(InTimeSignature, 0, 0, 0.0f));
+}
+
+void UFrameBasedMusicMap::SetInitialTempo(float Bpm)
+{
+	check(!TempoMap.IsEmpty());
+	TempoMap[0].SetBpm(Bpm);
+}
+
+float UFrameBasedMusicMap::GetInitialTempo() const
+{
+	return TempoMap.IsEmpty() ? 0.0f : TempoMap[0].Bpm();
+}
+
+void UFrameBasedMusicMap::SetInitialTimeSignature(const FFrameBasedTimeSignature& TimeSignature)
+{
+	check(!BarMap.IsEmpty());
+	BarMap[0].TimeSignature = TimeSignature;
+}
+
+FFrameBasedTimeSignature UFrameBasedMusicMap::GetInitialTimeSignature() const
+{
+	return BarMap.IsEmpty() ? FFrameBasedTimeSignature() : BarMap[0].TimeSignature;
+}
+
+double UFrameBasedMusicMap::GetFrame(const FMusicalTime& MusicalTime) const
+{
+	return FrameResolution.AsFrameTime(GetSeconds(MusicalTime)).AsDecimal();
+}
+
+double UFrameBasedMusicMap::GetFrame(int32 AtBar) const
+{
+	return FrameResolution.AsFrameTime(GetSeconds(AtBar)).AsDecimal();
+}
+
+FFrameTime UFrameBasedMusicMap::GetFrameTime(const FMusicalTime& MusicalTime) const
+{
+	return FrameResolution.AsFrameTime(GetSeconds(MusicalTime));
+}
+
+FFrameTime UFrameBasedMusicMap::GetFrameTime(int32 AtBar, float AtBeat) const
+{
+	return FrameResolution.AsFrameTime(GetSeconds(AtBar, AtBeat));
+}
+
+double UFrameBasedMusicMap::GetSeconds(const FMusicalTime& MusicalTime) const
+{
+	if (IsEmpty())
+	{
+		return 0.0;
+	}
+
+	check(MusicalTime.IsValid());
+
+	const FFrameBasedTimeSignaturePoint* TimeSignatureAtFrame = GetTimeSignaturePointForBar(MusicalTime.Bar);
+
+	int32 Tick = TimeSignatureAtFrame->MusicalTimeToTick(MusicalTime);
+
+	const FFrameBasedTempoPoint* TempoAtFrame = GetTempoPointForTick(Tick);
+
+	return TempoAtFrame->TickToSeconds(Tick);
+}
+
+double UFrameBasedMusicMap::GetSeconds(int32 AtBar, float AtBeat) const
+{
+	if (IsEmpty())
+	{
+		return 0.0;
+	}
+
+	const FFrameBasedTimeSignaturePoint* TimeSignatureAtFrame = GetTimeSignaturePointForBar(AtBar);
+
+	int32 Tick = TimeSignatureAtFrame->BarBeatToTick(AtBar, AtBeat);
+
+	const FFrameBasedTempoPoint* TempoAtFrame = GetTempoPointForTick(Tick);
+
+	return TempoAtFrame->TickToSeconds(Tick);
+}
+
+FString UFrameBasedMusicMap::MakeBarBeatStringForFrame(double DecimalFrame, FrameBasedMusicMap::EBarBeatFormat BarBeatFormat) const
+{
+	if (IsEmpty())
+	{
+		return FString(TEXT("0.0.00"));
+	}
+
+	const FFrameBasedTimeSignaturePoint* TimeSignatureAtFrame = GetTimeSignaturePointForFrame(DecimalFrame);
+	const FFrameBasedTempoPoint* TempoAtFrame = GetTempoPointForFrame(DecimalFrame);
+
+	return MakeBarBeatString(DecimalFrame, TimeSignatureAtFrame, TempoAtFrame, BarBeatFormat);
+}
+
+FString UFrameBasedMusicMap::MakeBarBeatStringForFrameTime(const FFrameTime& Frame, FrameBasedMusicMap::EBarBeatFormat BarBeatFormat) const
+{
+	return MakeBarBeatStringForFrame(Frame.AsDecimal(), BarBeatFormat);
+}
+
+FString UFrameBasedMusicMap::MakeBarBeatStringForMusicalTime(const FMusicalTime& MusicalTime, FrameBasedMusicMap::EBarBeatFormat BarBeatFormat) const
+{
+	FFrameTime FrameTime = GetFrameTime(MusicalTime);
+	return MakeBarBeatStringForFrameTime(FrameTime, BarBeatFormat);
+}
+
+FString UFrameBasedMusicMap::MakeBarBeatStringForSeconds(double Seconds, FrameBasedMusicMap::EBarBeatFormat BarBeatFormat) const
+{
+	return MakeBarBeatStringForFrameTime(FrameResolution.AsFrameTime(Seconds), BarBeatFormat);
+}
+
+float UFrameBasedMusicMap::GetTempoForFrame(double DecimalFrame) const
+{
+	return GetTempoForFrameTime(FFrameTime::FromDecimal(DecimalFrame));
+}
+
+float UFrameBasedMusicMap::GetTempoForFrameTime(const FFrameTime& FrameTime) const
+{
+	return GetTempoForSeconds(FrameResolution.AsSeconds(FrameTime));
+}
+
+float UFrameBasedMusicMap::GetTempoForSeconds(double Seconds) const
+{
+	if (TempoMap.IsEmpty())
+	{
+		return 0.0f;
+	}
+
+	if (Seconds <= 0.0)
+	{
+		return TempoMap[0].Bpm();
+	}
+
+	const FFrameBasedTempoPoint* TempoPoint = GetTempoPointForSeconds(Seconds);
+
+	return TempoPoint->Bpm();
+}
+
+FFrameBasedTimeSignature UFrameBasedMusicMap::GetTimeSignatureForFrame(double DecimalFrame) const
+{
+	const FFrameBasedTimeSignaturePoint* Point = GetTimeSignaturePointForFrame(DecimalFrame);
+	return Point ? Point->TimeSignature : FFrameBasedTimeSignature();
+}
+
+FFrameBasedTimeSignature UFrameBasedMusicMap::GetTimeSignatureForFrameTime(const FFrameTime& FrameTime) const
+{
+	const FFrameBasedTimeSignaturePoint* Point = GetTimeSignaturePointForFrameTime(FrameTime);
+	return Point ? Point->TimeSignature : FFrameBasedTimeSignature();
+}
+
+FFrameBasedTimeSignature UFrameBasedMusicMap::GetTimeSignatureForSeconds(double Seconds) const
+{
+	const FFrameBasedTimeSignaturePoint* Point = GetTimeSignaturePointForSeconds(Seconds);
+	return Point ? Point->TimeSignature : FFrameBasedTimeSignature();
+}
+
+const FFrameBasedTimeSignaturePoint* UFrameBasedMusicMap::GetTimeSignaturePointForTick(int32 Tick) const
+{
+	if (BarMap.IsEmpty())
+	{
+		return nullptr;
+	}
+
+	if (Tick <= 0)
+	{
+		return &BarMap[0];
+	}
+
+	const FFrameBasedTimeSignaturePoint* TimeSignatureAtTick = Algo::FindLastByPredicate(BarMap, [=](FFrameBasedTimeSignaturePoint const& Point) { return Tick >= Point.OnTick; });
+	return TimeSignatureAtTick ? TimeSignatureAtTick : &BarMap[0];
+}
+
+const FFrameBasedTimeSignaturePoint* UFrameBasedMusicMap::GetTimeSignaturePointForBar(int32 Bar) const
+{
+	if (BarMap.IsEmpty())
+	{
+		return nullptr;
+	}
+
+	if (Bar <= 0)
+	{
+		return &BarMap[0];
+	}
+
+	const FFrameBasedTimeSignaturePoint* TimeSignatureAtBar = Algo::FindLastByPredicate(BarMap, [=](FFrameBasedTimeSignaturePoint const& Point) { return Bar >= Point.OnBar; });
+	return TimeSignatureAtBar ? TimeSignatureAtBar : &BarMap[0];
+}
+
+const FFrameBasedTimeSignaturePoint* UFrameBasedMusicMap::GetTimeSignaturePointForFrame(double DecimalFrame) const
+{
+	if (BarMap.IsEmpty())
+	{
+		return nullptr;
+	}
+
+	if (DecimalFrame <= 0.0)
+	{
+		return &BarMap[0];
+	}
+
+	const FFrameBasedTimeSignaturePoint* TimeSignatureAtBar = Algo::FindLastByPredicate(BarMap, [=](FFrameBasedTimeSignaturePoint const& Point) { return DecimalFrame >= Point.OnFrame; });
+	return TimeSignatureAtBar ? TimeSignatureAtBar : &BarMap[0];
+}
+
+const FFrameBasedTimeSignaturePoint* UFrameBasedMusicMap::GetTimeSignaturePointForFrameTime(const FFrameTime& Frame) const
+{
+	return GetTimeSignaturePointForFrame(Frame.AsDecimal());
+}
+
+const FFrameBasedTimeSignaturePoint* UFrameBasedMusicMap::GetTimeSignaturePointForSeconds(double Seconds) const
+{
+	if (BarMap.IsEmpty() || TempoMap.IsEmpty())
+	{
+		return nullptr;
+	}
+
+	if (Seconds <= 0.0)
+	{
+		return &BarMap[0];
+	}
+
+	const FFrameBasedTempoPoint* TempoAtTime = GetTempoPointForSeconds(Seconds);
+	double Tick = TempoAtTime->SecondsToTick(Seconds);
+	return GetTimeSignaturePointForTick(Tick);
+}
+
+const FFrameBasedTempoPoint* UFrameBasedMusicMap::GetTempoPointForTick(int32 Tick) const
+{
+	if (TempoMap.IsEmpty())
+	{
+		return nullptr;
+	}
+
+	if (Tick <= 0)
+	{
+		return &TempoMap[0];
+	}
+
+	const FFrameBasedTempoPoint* TempoAtTick = Algo::FindLastByPredicate(TempoMap, [=](FFrameBasedTempoPoint const& Point) { return Tick >= Point.OnTick; });
+	return TempoAtTick ? TempoAtTick : &TempoMap[0];
+}
+
+const FFrameBasedTempoPoint* UFrameBasedMusicMap::GetTempoPointForFrame(double DecimalFrame) const
+{
+	return GetTempoPointForFrameTime(FFrameTime::FromDecimal(DecimalFrame));
+}
+
+const FFrameBasedTempoPoint* UFrameBasedMusicMap::GetTempoPointForFrameTime(const FFrameTime& Frame) const
+{
+	return GetTempoPointForSeconds(FrameResolution.AsSeconds(Frame));
+}
+
+const FFrameBasedTempoPoint* UFrameBasedMusicMap::GetTempoPointForSeconds(double Seconds) const
+{
+	if (TempoMap.IsEmpty())
+	{
+		return nullptr;
+	}
+
+	if (Seconds <= 0.0)
+	{
+		return &TempoMap[0];
+	}
+
+	float Ms = static_cast<float>(Seconds * 1000.0);
+	const FFrameBasedTempoPoint* TempoAtTime = Algo::FindLastByPredicate(TempoMap, [=](FFrameBasedTempoPoint const& Point) { return Ms >= Point.OnMs; });
+	return TempoAtTime ? TempoAtTime : &TempoMap[0];
+}
+
+FFrameTime UFrameBasedMusicMap::RemapFrameTime(FFrameTime SourceFrameTime, const UFrameBasedMusicMap* OldMap) const
+{
+	if (IsEmpty())
+	{
+		return SourceFrameTime;
+	}
+
+	FMusicalTime MusicalTime = OldMap->GetMusicalTime(SourceFrameTime);
+	const FFrameBasedTimeSignaturePoint* TimeSignaturePoint = GetTimeSignaturePointForBar(MusicalTime.Bar);
+	if (TimeSignaturePoint->TicksPerBar() != MusicalTime.TicksPerBar)
+	{
+		// Time Signature mismatch! Scale bar to bar!
+		MusicalTime.TickInBar = static_cast<int32>( static_cast<float>(MusicalTime.TickInBar * TimeSignaturePoint->TicksPerBar()) / static_cast<float>(MusicalTime.TicksPerBar));
+	}
+	MusicalTime.TicksPerBar = TimeSignaturePoint->TicksPerBar();
+	MusicalTime.TicksPerBeat = TimeSignaturePoint->TicksPerBeat();
+	return GetFrameTime(MusicalTime);
+}
+
+FMusicalTime UFrameBasedMusicMap::BarsToMusicalTime(double BarsFromStart) const
+{
+	FMusicalTime Result;
+	Result.Bar = FMath::FloorToInt32(BarsFromStart);
+	double FractionalBar = BarsFromStart - static_cast<double>(Result.Bar);
+
+	const FFrameBasedTimeSignaturePoint* TimeSignaturePoint = GetTimeSignaturePointForBar(Result.Bar);
+	Result.TicksPerBar = TimeSignaturePoint->TicksPerBar();
+	Result.TicksPerBeat = TimeSignaturePoint->TicksPerBeat();
+	Result.TickInBar = FMath::Clamp(FMath::RoundToInt32( static_cast<double>(Result.TicksPerBar) * FractionalBar), 0, Result.TicksPerBar - 1);
+	return Result;
+}
+
+FMusicalTime UFrameBasedMusicMap::Quantize(const FMusicalTime& MusicalTime, int32 QuantizationInterval, EQuantizeDirection Direction) const
+{
+	int32 TicksInQuantizationInterval = QuantizationInterval;
+	switch (QuantizationInterval)
+	{
+	case FrameBasedMusicMap::EWholeQuantizationTicks::Bar:
+		TicksInQuantizationInterval = MusicalTime.TicksPerBar;
+		break;
+	case FrameBasedMusicMap::EWholeQuantizationTicks::Beat:
+		TicksInQuantizationInterval = MusicalTime.TicksPerBeat;
+		break;
+	}
+
+	const FFrameBasedTimeSignaturePoint* TimeSignaturePoint = GetTimeSignaturePointForBar(MusicalTime.Bar);
+
+	int32 TotalTicks = TimeSignaturePoint->MusicalTimeToTick(MusicalTime);
+
+	int32 TicksInTimeSignature = TotalTicks - TimeSignaturePoint->OnTick;
+
+	int32 Intervals = 0;
+	switch (Direction)
+	{
+	default:
+	case EQuantizeDirection::Nearest:
+		Intervals = FMath::RoundToInt32(static_cast<float>(TicksInTimeSignature) / static_cast<float>(TicksInQuantizationInterval));
+		break;
+	case EQuantizeDirection::Ceiling:
+		Intervals = FMath::CeilToInt32(static_cast<float>(TicksInTimeSignature) / static_cast<float>(TicksInQuantizationInterval));
+		break;
+	case EQuantizeDirection::Floor:
+		Intervals = FMath::FloorToInt32(static_cast<float>(TicksInTimeSignature) / static_cast<float>(TicksInQuantizationInterval));
+		break;
+	}
+
+	int32 QuantizedTicksInTimeSignature = Intervals * TicksInQuantizationInterval;
+
+	int32 FinalTick = TimeSignaturePoint->OnTick + QuantizedTicksInTimeSignature;
+
+	TimeSignaturePoint = GetTimeSignaturePointForTick(FinalTick);
+
+	return TimeSignaturePoint->TickToMusicalTime(FinalTick);
+}
+
+FMusicalTime UFrameBasedMusicMap::Add(const FMusicalTime& MusicalTime, int32 QuantizationInterval) const
+{
+	const FFrameBasedTimeSignaturePoint* OriginalTimeSignature = GetTimeSignaturePointForBar(MusicalTime.Bar);
+	int32 AbsoluteTick = OriginalTimeSignature->MusicalTimeToTick(MusicalTime);
+	if (FMath::Abs(QuantizationInterval) == FrameBasedMusicMap::EWholeQuantizationTicks::Bar)
+	{
+		AbsoluteTick += (QuantizationInterval < 0) ? MusicalTime.TicksPerBar : -MusicalTime.TicksPerBar;
+	}
+	else if (FMath::Abs(QuantizationInterval) == FrameBasedMusicMap::EWholeQuantizationTicks::Beat)
+	{
+		AbsoluteTick += (QuantizationInterval < 0) ? MusicalTime.TicksPerBeat : -MusicalTime.TicksPerBeat;
+	}
+	else
+	{
+		AbsoluteTick += QuantizationInterval;
+	}
+	const FFrameBasedTimeSignaturePoint* NewTimeSignature = GetTimeSignaturePointForTick(AbsoluteTick);
+	return NewTimeSignature->TickToMusicalTime(AbsoluteTick);
+}
+
+void UFrameBasedMusicMap::AddTimeSignature(int32 AtTick, int32 AtBar, int16 Numerator, int16 Denominator)
+{
+	if (ensure(!BarMap.IsEmpty() || AtTick == 0))
+	{
+		if (ensure(AtTick == 0 || BarMap.Last().OnTick < AtTick))
+		{
+			BarMap.Add(FFrameBasedTimeSignaturePoint({Numerator, Denominator}, AtBar, AtTick, GetFrame(AtBar)));
+		}
+	}
+}
+
+void UFrameBasedMusicMap::AddTempo(int32 AtTick, float AtMs, float Bpm)
+{
+	if (ensure(TempoMap.IsEmpty() || AtTick >= TempoMap.Last().OnTick))
+	{
+		float ExpectedMs = TempoMap.IsEmpty() ? 0.0f : static_cast<float>(TempoMap.Last().TickToSeconds(static_cast<double>(AtTick)) * 1000.0);
+		TempoMap.Add(FFrameBasedTempoPoint(Bpm, AtTick, ExpectedMs));
+	}
+}
+
+TArray<FSubdivisionPoint> UFrameBasedMusicMap::GatherSubdivisonsForTimespan(const FFrameTime& Start, const FFrameTime& End, const FFrameTime& MinSpacing, int32 QuantizationInterval) const
+{
+	TArray<FSubdivisionPoint> Result;
+
+	FMusicalTime StartMusicalTime = Quantize(GetMusicalTime(Start), QuantizationInterval, EQuantizeDirection::Ceiling);
+	FMusicalTime EndMusicalTime = Quantize(GetMusicalTime(End), QuantizationInterval, EQuantizeDirection::Floor);
+
+	if (StartMusicalTime >= EndMusicalTime)
+	{
+		return Result;
+	}
+
+
+	Result.Add(MakeSubdivisionPoint(StartMusicalTime));
+	FFrameTime LastEntryFrameTime = GetFrameTime(StartMusicalTime);
+	int32 LastBarNumber = StartMusicalTime.Bar;
+	StartMusicalTime = Add(StartMusicalTime, QuantizationInterval);
+	if (StartMusicalTime.Bar != LastBarNumber)
+	{
+		StartMusicalTime.TickInBar = 0;
+		LastBarNumber = StartMusicalTime.Bar;
+	}
+	while (StartMusicalTime <= EndMusicalTime)
+	{
+		FFrameTime ThisEntryFrameTime = GetFrameTime(StartMusicalTime);
+		if ((ThisEntryFrameTime - LastEntryFrameTime) < MinSpacing)
+		{
+			if (StartMusicalTime.TickInBar == 0)
+			{
+				// This is a downbeat. We are too close to the last, BUT we 
+				// take precedence. So blow away the last...
+				Result.Pop(EAllowShrinking::No);
+			}
+			else
+			{
+				StartMusicalTime = Add(StartMusicalTime, QuantizationInterval);
+				if (StartMusicalTime.Bar != LastBarNumber)
+				{
+					StartMusicalTime.TickInBar = 0;
+					LastBarNumber = StartMusicalTime.Bar;
+				}
+				continue;
+			}
+		}
+		Result.Add(MakeSubdivisionPoint(StartMusicalTime));
+		LastEntryFrameTime = GetFrameTime(StartMusicalTime);
+		StartMusicalTime = Add(StartMusicalTime, QuantizationInterval);
+		if (StartMusicalTime.Bar != LastBarNumber)
+		{
+			StartMusicalTime.TickInBar = 0;
+			LastBarNumber = StartMusicalTime.Bar;
+		}
+	}
+	
+	return Result;
+}
+
+FMusicalTime UFrameBasedMusicMap::GetMusicalTime(const FFrameTime& FrameTime) const
+{
+	if (IsEmpty())
+	{
+		return FMusicalTime();
+	}
+
+	const FFrameBasedTimeSignaturePoint* TimeSignatureAtFrame = GetTimeSignaturePointForFrameTime(FrameTime);
+	const FFrameBasedTempoPoint* TempoAtFrame = GetTempoPointForFrameTime(FrameTime);
+
+	double Seconds = FrameResolution.AsSeconds(FrameTime);
+	float Tick = TempoAtFrame->SecondsToTick(Seconds);
+	return TimeSignatureAtFrame->TickToMusicalTime(Tick);
+}
+
+FSubdivisionPoint UFrameBasedMusicMap::MakeSubdivisionPoint(const FMusicalTime& MusicalTime) const
+{
+	FSubdivisionPoint Result;
+	Result.MusicalTime = MusicalTime;
+	Result.FrameTime = GetFrameTime(MusicalTime);
+	Result.PrettyString = MakeBarBeatStringForMusicalTime(MusicalTime, FrameBasedMusicMap::EBarBeatFormat::Position);
+	return Result;
+}
+
+FString UFrameBasedMusicMap::MakeBarBeatString(double DecimalFrame, const FFrameBasedTimeSignaturePoint* TimeSignature, const FFrameBasedTempoPoint* Tempo, FrameBasedMusicMap::EBarBeatFormat BarBeatFormat) const
+{
+	double Seconds = FrameResolution.AsSeconds(FFrameTime::FromDecimal(DecimalFrame));
+	float Tick = Tempo->SecondsToTick(Seconds);
+	FMusicalTime MusicalTime = TimeSignature->TickToMusicalTime(Tick);
+	// for display to users... bars and beats are 1 based!
+	MusicalTime.Bar += (BarBeatFormat == FrameBasedMusicMap::EBarBeatFormat::Position ? 1 : 0);
+	double BeatAsDouble = FMath::RoundToDouble((MusicalTime.BeatInBar() + (BarBeatFormat == FrameBasedMusicMap::EBarBeatFormat::Position ? 1.0f : 0.0f)) * 100.0f);
+	if (BeatAsDouble >= ((TimeSignature->TimeSignature.Numerator + (BarBeatFormat == FrameBasedMusicMap::EBarBeatFormat::Position ? 1 : 0)) * 100))
+	{
+		MusicalTime.Bar++;
+		BeatAsDouble = BarBeatFormat == FrameBasedMusicMap::EBarBeatFormat::Position ? 100.0 : 0.0;
+	}
+	return FString::Printf(TEXT("%d.%.2f"), MusicalTime.Bar, BeatAsDouble / 100.0);
+}
+#undef LOCTEXT_NAMESPACE
diff --git a/Engine/Plugins/Runtime/MusicEnvironment/Source/MusicEnvironment/Private/MusicClockSourceManager.cpp b/Engine/Plugins/Runtime/MusicEnvironment/Source/MusicEnvironment/Private/MusicClockSourceManager.cpp
new file mode 100644
--- /dev/null
+++ b/Engine/Plugins/Runtime/MusicEnvironment/Source/MusicEnvironment/Private/MusicClockSourceManager.cpp
@@ -0,0 +1,155 @@
+// Copyright Epic Games, Inc. All Rights Reserved.
+
+#include "MusicClockSourceManager.h"
+
+#include "MusicEnvironmentSubsystem.h" 
+
+TScriptInterface<IMusicEnvironmentClockSource> UMusicClockSourceManager::FindClock(const FGameplayTag& Tag, bool bExactMatch)
+{
+	CleanMap();
+	if (const TWeakInterfacePtr<IMusicEnvironmentClockSource>* Clock = TaggedClocks.Find(Tag))
+	{
+		return Clock->ToScriptInterface();;
+	}
+	if (!bExactMatch)
+	{
+		int32 NumberOfMatches = 0;
+		TWeakInterfacePtr<IMusicEnvironmentClockSource> MatchingClock;
+		FGameplayTag MatchingTag;
+		for (auto& It : TaggedClocks)
+		{
+			if (It.Key.MatchesTag(Tag))
+			{
+				++NumberOfMatches;
+				if (!MatchingClock.IsValid())
+				{
+					MatchingClock = It.Value;
+					MatchingTag = It.Key;
+				}
+			}
+		}
+		if (MatchingClock.IsValid())
+		{
+			if (NumberOfMatches > 1)
+			{ 
+				UE_LOG(LogMusicEnvironment, Warning, TEXT("FindClock found %d clocks that match the tag '%s'. Returning clock tagged '%s'."), NumberOfMatches, *Tag.ToString(), *MatchingTag.ToString());
+			}
+			return MatchingClock.ToScriptInterface();
+		}
+	}
+	return nullptr;
+}
+
+void UMusicClockSourceManager::AddTaggedClock(const FGameplayTag& Tag, TScriptInterface<IMusicEnvironmentClockSource> InClock)
+{
+	UObject* AsObject = InClock.GetObject();
+	AddTaggedClock(Tag, AsObject);
+}
+
+void UMusicClockSourceManager::AddTaggedClock(const FGameplayTag& Tag, UObject* InClock)
+{
+	CleanMap();
+	TWeakInterfacePtr<IMusicEnvironmentClockSource> WeakInterfacePtr(InClock);
+	if (!WeakInterfacePtr.IsValid())
+	{
+		return;
+	}
+	if (TaggedClocks.Contains(Tag))
+	{
+		if (WeakInterfacePtr != TaggedClocks[Tag])
+		{
+			UE_LOG(LogMusicEnvironment, Warning, TEXT("The MusicClockSourceManager is already tracking a clock tagged '%s'. It will be replaced by this new request."), *Tag.ToString());
+			TaggedClocks[Tag] = WeakInterfacePtr;
+		}
+	}
+	else
+	{
+		TaggedClocks.Add(Tag, WeakInterfacePtr);
+	}
+}
+
+void UMusicClockSourceManager::RemoveTaggedClock(TScriptInterface<IMusicEnvironmentClockSource> InClock)
+{
+	UObject* AsObject = InClock.GetObject();
+	RemoveTaggedClock(AsObject);
+}
+
+void UMusicClockSourceManager::RemoveTaggedClock(UObject* InClock)
+{
+	CleanMap();
+	TWeakInterfacePtr<IMusicEnvironmentClockSource> WeakInterfacePtr(InClock);
+	if (WeakInterfacePtr.IsValid())
+	{
+		TaggedClocks = TaggedClocks.FilterByPredicate([&](const TPair<FGameplayTag, TWeakInterfacePtr<IMusicEnvironmentClockSource>>& ClockEntry) { return ClockEntry.Value != WeakInterfacePtr; });
+	}
+}
+
+void UMusicClockSourceManager::RemoveClockWithTag(const FGameplayTag& Tag)
+{
+	CleanMap();
+	TaggedClocks.Remove(Tag);
+}
+
+TScriptInterface<IMusicEnvironmentClockSource> UMusicClockSourceManager::GetGlobalMusicClockAuthority()
+{
+	CleanStack();
+	if (GlobalClockSourceStack.IsEmpty())
+	{
+		return nullptr;
+	}
+	return GlobalClockSourceStack.Last().ToScriptInterface();
+}
+
+bool UMusicClockSourceManager::PushGlobalMusicClockAuthority(TScriptInterface<IMusicEnvironmentClockSource> InClock)
+{
+	UObject* AsObject = InClock.GetObject();
+	return PushGlobalMusicClockAuthority(AsObject);
+}
+
+bool UMusicClockSourceManager::PushGlobalMusicClockAuthority(UObject* InClock)
+{
+	CleanStack();
+	TWeakInterfacePtr<IMusicEnvironmentClockSource> WeakInterfacePtr(InClock);
+	if (!WeakInterfacePtr.IsValid())
+	{
+		return false;
+	}
+	GlobalClockSourceStack.Add(WeakInterfacePtr);
+	return true;
+}
+
+void UMusicClockSourceManager::RemoveGlobalClockAuthority(TScriptInterface<IMusicEnvironmentClockSource> InClock)
+{
+	UObject* AsObject = InClock.GetObject();
+	RemoveGlobalClockAuthority(AsObject);
+}
+
+void UMusicClockSourceManager::RemoveGlobalClockAuthority(UObject* InClock)
+{
+	CleanStack();
+	TWeakInterfacePtr<IMusicEnvironmentClockSource> WeakInterfacePtr(InClock);
+	if (!WeakInterfacePtr.IsValid())
+	{
+		return;
+	}
+	GlobalClockSourceStack.RemoveAll([&](const TWeakInterfacePtr<IMusicEnvironmentClockSource>& ClockEntry) { return ClockEntry == WeakInterfacePtr; });
+}
+
+void UMusicClockSourceManager::PopMusicClockAuthority()
+{
+	CleanStack();
+	if (!GlobalClockSourceStack.IsEmpty())
+	{
+		GlobalClockSourceStack.Pop(EAllowShrinking::No);
+	}
+}
+
+void UMusicClockSourceManager::CleanMap()
+{
+	TaggedClocks = TaggedClocks.FilterByPredicate([](const TPair<FGameplayTag, TWeakInterfacePtr<IMusicEnvironmentClockSource>>& ClockEntry){ return ClockEntry.Value.IsValid(); });
+}
+
+void UMusicClockSourceManager::CleanStack()
+{
+	GlobalClockSourceStack.RemoveAll([](const TWeakInterfacePtr<IMusicEnvironmentClockSource>& ClockEntry){ return !ClockEntry.IsValid(); });
+}
diff --git a/Engine/Plugins/Runtime/MusicEnvironment/Source/MusicEnvironment/Private/MusicEnvironmentMetronome.cpp b/Engine/Plugins/Runtime/MusicEnvironment/Source/MusicEnvironment/Private/MusicEnvironmentMetronome.cpp
new file mode 100644
--- /dev/null
+++ b/Engine/Plugins/Runtime/MusicEnvironment/Source/MusicEnvironment/Private/MusicEnvironmentMetronome.cpp
@@ -0,0 +1,19 @@
+// Copyright Epic Games, Inc. All Rights Reserved.
+
+#include "MusicEnvironmentMetronome.h"
+
+// Add default functionality here for any IMovieSceneMetronome functions that are not pure virtual.
+void IMusicEnvironmentMetronome::SetMusicMap(UFrameBasedMusicMap* InMusicMap)
+{
+	MusicMap = TStrongObjectPtr<UFrameBasedMusicMap>(InMusicMap);
+	OnMusicMapSet();
+}
+
+bool IMusicEnvironmentMetronome::SetTempo(const float Bpm)
+{
+	if (MusicMap.IsValid() && !MusicMap->IsSimple())
+	{
+		return false;
+	}
+	return OnSetTempo(Bpm);
+}
diff --git a/Engine/Plugins/Runtime/MusicEnvironment/Source/MusicEnvironment/Private/MusicEnvironmentModule.cpp b/Engine/Plugins/Runtime/MusicEnvironment/Source/MusicEnvironment/Private/MusicEnvironmentModule.cpp
new file mode 100644
--- /dev/null
+++ b/Engine/Plugins/Runtime/MusicEnvironment/Source/MusicEnvironment/Private/MusicEnvironmentModule.cpp
@@ -0,0 +1,7 @@
+// Copyright Epic Games, Inc. All Rights Reserved.
+
+#include "MusicEnvironmentModule.h"
+
+#include "Logging/MessageLog.h"
+
+IMPLEMENT_MODULE(FMusicEnvironmentModule, MusicEnvironment);
diff --git a/Engine/Plugins/Runtime/MusicEnvironment/Source/MusicEnvironment/Private/MusicEvironmentSubsystem.cpp b/Engine/Plugins/Runtime/MusicEnvironment/Source/MusicEnvironment/Private/MusicEvironmentSubsystem.cpp
new file mode 100644
--- /dev/null
+++ b/Engine/Plugins/Runtime/MusicEnvironment/Source/MusicEnvironment/Private/MusicEvironmentSubsystem.cpp
@@ -0,0 +1,48 @@
+// Copyright Epic Games, Inc. All Rights Reserved.
+
+#include "MusicEnvironmentSubsystem.h"
+
+#include "MusicClockSourceManager.h"
+#include "MusicEnvironmentMetronome.h"
+
+#include "Engine/Engine.h"
+
+DEFINE_LOG_CATEGORY(LogMusicEnvironment);
+
+UMusicEnvironmentSubsystem& UMusicEnvironmentSubsystem::Get()
+{
+	return *GEngine->GetEngineSubsystem< UMusicEnvironmentSubsystem>();
+}
+
+UMusicClockSourceManager* UMusicEnvironmentSubsystem::GetClockSourceManager()
+{
+	if (!ClockSourceManager)
+	{
+		ClockSourceManager = NewObject<UMusicClockSourceManager>(this);
+	}
+	return ClockSourceManager;
+}
+
+bool UMusicEnvironmentSubsystem::SetMetronomeClass(const UClass* InMetronomeType)
+{
+	if (!InMetronomeType->ImplementsInterface(UMusicEnvironmentMetronome::StaticClass()))
+	{
+		return false;
+	}
+	MetronomeType = InMetronomeType;
+	return true;
+}
+
+bool UMusicEnvironmentSubsystem::CanSpawnMetronome() const
+{
+	return MetronomeType.IsValid();
+}
+
+TScriptInterface<IMusicEnvironmentMetronome> UMusicEnvironmentSubsystem::SpawnMetronome(UObject* Outer)
+{
+	if (!MetronomeType.IsValid())
+	{
+		return nullptr;
+	}
+	return TScriptInterface<IMusicEnvironmentMetronome>(NewObject<UObject>(Outer, MetronomeType.Get()));
+}
diff --git a/Engine/Plugins/Runtime/MusicEnvironment/Source/MusicEnvironment/Public/FrameBasedMusicMap.h b/Engine/Plugins/Runtime/MusicEnvironment/Source/MusicEnvironment/Public/FrameBasedMusicMap.h
new file mode 100644
--- /dev/null
+++ b/Engine/Plugins/Runtime/MusicEnvironment/Source/MusicEnvironment/Public/FrameBasedMusicMap.h
@@ -0,0 +1,415 @@
+// Copyright Epic Games, Inc. All Rights Reserved.
+
+#pragma once
+
+#include "CoreTypes.h"
+#include "Misc/FrameRate.h"
+#include "Misc/MusicalTime.h"
+#include "UObject/Interface.h"
+
+#include "FrameBasedMusicMap.generated.h"
+
+struct FFrameTime;
+
+namespace FrameBasedMusicMap
+{
+	enum ETripletQuantizationTicks : int32
+	{
+		ThirtysecondNoteTriplet = (MusicalTime::TicksPerQuarterNote / 4) / 3,	// 80
+		SixteenthNoteTriplet    = (MusicalTime::TicksPerQuarterNote / 2) / 3,	// 160
+		EighthNoteTriplet       = (MusicalTime::TicksPerQuarterNote) / 3,		// 320
+		QuarterNoteTriplet      = (MusicalTime::TicksPerQuarterNote * 2) / 3,	// 640
+		HalfNoteTriplet         = (MusicalTime::TicksPerQuarterNote * 4) / 3,	// 1280
+		WholeNoteTriplet        = (MusicalTime::TicksPerQuarterNote * 8) / 3,	// 2560
+	};
+
+	enum EWholeQuantizationTicks : int32
+	{
+		ThirtysecondNote        = MusicalTime::TicksPerQuarterNote / 8,			// 120
+		SixteenthNote           = MusicalTime::TicksPerQuarterNote / 4,			// 240
+		EighthNote              = MusicalTime::TicksPerQuarterNote / 2,			// 480
+		QuarterNote	            = MusicalTime::TicksPerQuarterNote,				// 960
+		HalfNote                = MusicalTime::TicksPerQuarterNote * 2,			// 1920
+		WholeNote               = MusicalTime::TicksPerQuarterNote * 4,			// 3840
+		Bar                     = std::numeric_limits<int32>::min()+1,          // Ticks depends on Time Signature
+		Beat                    = std::numeric_limits<int32>::min()+2           // Ticks depends on Time Signature
+	};
+
+	enum EDottedQuantizationTicks : int32
+	{
+		DottedThirtysecondNote  = (MusicalTime::TicksPerQuarterNote / 16) * 3,	// 180
+		DottedSixteenthNote     = (MusicalTime::TicksPerQuarterNote / 8) * 3,	// 360
+		DottedEighthNote        = (MusicalTime::TicksPerQuarterNote / 4) * 3,	// 720
+		DottedQuarterNote       = (MusicalTime::TicksPerQuarterNote / 2) * 3,	// 1440
+		DottedHalfNote          = (MusicalTime::TicksPerQuarterNote) * 3,		// 2880
+		DottedWholeNote         = (MusicalTime::TicksPerQuarterNote) * 6,		// 5760
+	};
+
+	MUSICENVIRONMENT_API FText GetQuantizationName(int32 Quantization);
+	extern MUSICENVIRONMENT_API TArray<TSharedPtr<int32>> QuantizationOptionsSource;
+
+	enum class EBarBeatFormat
+	{
+		Position,
+		Duration,
+		NotAllowed // Don't display bars.beats! Fall back to seconds.
+	};
+}
+
+struct FSubdivisionPoint
+{
+	FFrameTime FrameTime;
+	FMusicalTime MusicalTime;
+	FString PrettyString;
+};
+
+USTRUCT()
+struct MUSICENVIRONMENT_API FFrameBasedTempoPoint
+{
+	GENERATED_BODY()
+
+	/**
+	 * NOTE: This structure treats BPM/Tempo exactly like Midi does to avoid any
+	 * floating point precision problems and differences between Midi code and 
+	 * this code. It is for that reason that Tempo is stored here a "microseconds
+	 * per quarter note" (MPQ) and not a float Bpm. 
+	 */
+
+	FFrameBasedTempoPoint()
+	{}
+
+	FFrameBasedTempoPoint(float InBpm, int32 InTick, float InAtMillisecond)
+		: OnTick(InTick)
+		, OnMs(InAtMillisecond)
+	{
+		SetBpm(InBpm);
+	}
+
+	void SetBpm(float InBpm)
+	{
+		MicrosecondsPerQuarterNote = InBpm >= 0.0f ? static_cast<int32>(60000000.0f / InBpm) : 0;
+	}
+
+	UPROPERTY()
+	int32 MicrosecondsPerQuarterNote = 500000; // 120 bpm default.
+	UPROPERTY()
+	int32 OnTick = 0;
+	UPROPERTY()
+	float OnMs = 0.0;
+
+	float Bpm() const
+	{
+		return MicrosecondsPerQuarterNote == 0 ? 0 : 60000000.0f / static_cast<float>(MicrosecondsPerQuarterNote);
+	}
+
+	int32 MsToTick(float Ms) const
+	{
+		if (Ms == 0.0f)
+		{
+			return 0;
+		}
+
+		return FMath::RoundToInt32(static_cast<double>(OnTick) +
+								static_cast<double>(Ms - OnMs) * 1000.0 / 
+								static_cast<double>(MicrosecondsPerQuarterNote) *
+								static_cast<double>(MusicalTime::TicksPerQuarterNote));
+	}
+
+	int32 SecondsToTick(double Second) const
+	{
+		return MsToTick(static_cast<float>(Second * 1000.0));
+	}
+
+	double TickToSeconds(int32 Tick) const
+	{
+		return static_cast<double>(TickToMs(Tick))/1000.0;
+	}
+
+	float TickToMs(int32 Tick) const
+	{
+		if (Tick == 0)
+		{
+			return 0.0f;
+		}
+		float MsPerTick = static_cast<float>(MicrosecondsPerQuarterNote) / (static_cast<float>(MusicalTime::TicksPerQuarterNote) * 1000.0f);
+		float DeltaTick = static_cast<float>(Tick - OnTick);
+		return OnMs + (DeltaTick * MsPerTick);
+	}
+};
+
+USTRUCT()
+struct MUSICENVIRONMENT_API FFrameBasedTimeSignature
+{
+	GENERATED_BODY()
+
+	FFrameBasedTimeSignature()
+	{}
+	FFrameBasedTimeSignature(int32 InNumerator, int32 InDenominator)
+		: Numerator(static_cast<int16>(InNumerator))
+		, Denominator(static_cast<int16>(InDenominator))
+	{}
+
+	UPROPERTY()
+	int16 Numerator = 4;
+	UPROPERTY()
+	int16 Denominator = 4;
+
+	bool operator==(const FFrameBasedTimeSignature& Other)
+	{
+		return Numerator == Other.Numerator && Denominator == Other.Denominator;
+	}
+
+	bool operator!=(const FFrameBasedTimeSignature& Other)
+	{
+		return Numerator != Other.Numerator || Denominator != Other.Denominator;
+	}
+};
+
+USTRUCT()
+struct MUSICENVIRONMENT_API FFrameBasedTimeSignaturePoint
+{
+	GENERATED_BODY()
+
+	FFrameBasedTimeSignaturePoint()
+	{}
+
+	FFrameBasedTimeSignaturePoint(const FFrameBasedTimeSignature& InTimeSignature, int32 InOnBar, int32 InOnTick, double InOnFrame)
+		: TimeSignature(InTimeSignature)
+		, OnBar(InOnBar)
+		, OnTick(InOnTick)
+		, OnFrame(InOnFrame)
+	{
+	}
+
+	UPROPERTY()
+	FFrameBasedTimeSignature TimeSignature;
+	UPROPERTY()
+	int32 OnBar = 0;
+	UPROPERTY()
+	int32 OnTick = 0;
+	UPROPERTY()
+	double OnFrame = 0.0;
+
+	int32 TicksPerBeat() const
+	{
+		return MusicalTime::TicksPerQuarterNote * 4 / TimeSignature.Denominator;
+	}
+
+	int32 TicksPerBar() const
+	{
+		return TicksPerBeat() * TimeSignature.Numerator;
+	}
+
+	FMusicalTime TickToMusicalTime(int32 InTick) const
+	{
+		FMusicalTime Result(0, 0, TicksPerBar(), TicksPerBeat());
+		int32 TicksPast = InTick - OnTick;
+		int32 BarsPassed = TicksPast / TicksPerBar();
+		if (TicksPast >= 0)
+		{
+			Result.Bar = OnBar + BarsPassed;
+			Result.TickInBar = TicksPast - (BarsPassed * TicksPerBar());
+		}
+		else
+		{
+			Result.Bar = OnBar + (BarsPassed - 1);
+			Result.TickInBar = TicksPast - ((-1 + BarsPassed) * TicksPerBar());
+		}
+		return Result;
+	}
+
+	int32 MusicalTimeToTick(const FMusicalTime& MusicalTime) const
+	{
+		int32 BarDelta = MusicalTime.Bar - OnBar;
+		return OnTick + (BarDelta * MusicalTime.TicksPerBar) + MusicalTime.TickInBar;
+	}
+
+	int32 BarToTick(int32 Bar) const
+	{
+		int32 BarDelta = Bar - OnBar;
+		return OnTick + (BarDelta * TicksPerBar());
+	}
+
+	int32 BarBeatToTick(int32 Bar, float Beat) const
+	{
+		int32 BarDelta = Bar - OnBar;
+		float TicksIntoBar = Beat * static_cast<float>(TicksPerBeat());
+		return OnTick + (BarDelta * TicksPerBar()) + FMath::RoundToInt32(TicksIntoBar);
+	}
+};
+
+UCLASS()
+class MUSICENVIRONMENT_API UFrameBasedMusicMap : public UObject
+{
+	GENERATED_BODY()
+
+public:
+	UFrameBasedMusicMap()
+		: FrameResolution(0,0)
+	{}
+
+	UFrameBasedMusicMap(const FFrameRate& InFrameResolution)
+		: FrameResolution(InFrameResolution)
+	{
+		
+	}
+
+	UFrameBasedMusicMap(const FFrameRate& InFrameResolution, float InBpm, const FFrameBasedTimeSignature& TimeSignature)
+		: FrameResolution(InFrameResolution)
+	{
+		TempoMap.Add(FFrameBasedTempoPoint(InBpm, 0, 0.0f));
+		BarMap.Add(FFrameBasedTimeSignaturePoint(TimeSignature, 0, 0, 0.0));
+	}
+
+	void SetFrameResolution(const FFrameRate& NewFrameResolution);
+
+	void Init(float InBpm, const FFrameBasedTimeSignature& InTimeSignature);
+
+	void SetInitialTempo(float Bpm);
+	float GetInitialTempo() const;
+	void SetInitialTimeSignature(const FFrameBasedTimeSignature& TimeSignature);
+	FFrameBasedTimeSignature GetInitialTimeSignature() const;
+
+	void Clear()
+	{
+		TempoMap.Empty();
+		BarMap.Empty();
+	}
+
+	bool IsEmpty() const { return TempoMap.IsEmpty() || BarMap.IsEmpty(); }
+	bool IsSimple() const { return TempoMap.Num() == 1 && BarMap.Num() == 1; }
+
+	double GetFrame(const FMusicalTime& MusicalTime) const;
+	double GetFrame(int32 AtBar) const;
+	FFrameTime GetFrameTime(const FMusicalTime& MusicalTime) const;
+	FFrameTime GetFrameTime(int32 AtBar, float AtBeat) const;
+	double GetSeconds(const FMusicalTime& MusicalTime) const;
+	double GetSeconds(int32 AtBar, float AtBeat = 0.0f) const;
+
+	FString MakeBarBeatStringForFrame(double DecimalFrame, FrameBasedMusicMap::EBarBeatFormat BarBeatFormat) const;
+	FString MakeBarBeatStringForFrameTime(const FFrameTime& Frame, FrameBasedMusicMap::EBarBeatFormat BarBeatFormat) const;
+	FString MakeBarBeatStringForMusicalTime(const FMusicalTime& MusicalTime, FrameBasedMusicMap::EBarBeatFormat BarBeatFormat) const;
+	FString MakeBarBeatStringForSeconds(double Seconds, FrameBasedMusicMap::EBarBeatFormat BarBeatFormat) const;
+
+	FMusicalTime GetMusicalTime(const FFrameTime& FrameTime) const;
+
+	float GetTempoForFrame(double DecimalFrame) const;
+	float GetTempoForFrameTime(const FFrameTime& FrameTime) const;
+	float GetTempoForSeconds(double Seconds) const;
+
+	FFrameBasedTimeSignature GetTimeSignatureForFrame(double DecimalFrame) const;
+	FFrameBasedTimeSignature GetTimeSignatureForFrameTime(const FFrameTime& FrameTime) const;
+	FFrameBasedTimeSignature GetTimeSignatureForSeconds(double Seconds) const;
+
+	FFrameTime RemapFrameTime(FFrameTime SourceTime, const UFrameBasedMusicMap* OldMap) const;
+
+	FMusicalTime BarsToMusicalTime(double BarsFromStart) const;
+
+	enum class EQuantizeDirection
+	{
+		Nearest,
+		Ceiling,
+		Floor
+	};
+	FMusicalTime Quantize(const FMusicalTime& MusicalTime, int32 QuantizationInterval, EQuantizeDirection Direction = EQuantizeDirection::Nearest) const;
+
+	/**
+	* FMusicalTime Add(const FMusicalTime& MusicalTime, int32 QuantizationInterval) const;
+	* 
+	* Increment the supplied MusicalTime by the specified QuantizationInterval. 
+	* NOTE: This does NOT quantize anything. QuantizationInterval is just used 
+	* as a convenient way to specify a musical subdivision. You can use any of the 
+	* TripletQuantizationTicks, WholeQuantizationTicks, or DottedQuantizationTicks
+	* enum values. AND... you can negate them to go back in time.
+	* 
+	* @param MusicalTime			Time to be added to
+	
+	* @param QuantizationInterval	One of the values from...
+	*                               TripletQuantizationTicks
+	*								WholeQuantizationTicks
+	*								DottedQuantizationTicks
+	* 
+	*                                 
+	*/
+	FMusicalTime Add(const FMusicalTime& MusicalTime, int32 QuantizationInterval) const;
+
+	UPROPERTY()
+	TArray<FFrameBasedTempoPoint> TempoMap;
+	UPROPERTY()
+	TArray<FFrameBasedTimeSignaturePoint> BarMap;
+	UPROPERTY()
+	FFrameRate FrameResolution;
+
+	void AddTimeSignature(int32 AtTick, int32 AtBar, int16 Numerator, int16 Denominator);
+	void AddTempo(int32 AtTick, float AtMs, float Bpm);
+
+	TArray<FSubdivisionPoint> GatherSubdivisonsForTimespan(const FFrameTime& Start, const FFrameTime& End, const FFrameTime& MinSpacing, int32 QuantizationInterval) const;
+
+private:
+	const FFrameBasedTimeSignaturePoint* GetTimeSignaturePointForTick(int32 Tick) const;
+	const FFrameBasedTimeSignaturePoint* GetTimeSignaturePointForBar(int32 Bar) const;
+	const FFrameBasedTimeSignaturePoint* GetTimeSignaturePointForFrame(double DecimalFrame) const;
+	const FFrameBasedTimeSignaturePoint* GetTimeSignaturePointForFrameTime(const FFrameTime& Frame) const;
+	const FFrameBasedTimeSignaturePoint* GetTimeSignaturePointForSeconds(double Seconds) const;
+
+	const FFrameBasedTempoPoint* GetTempoPointForTick(int32 Tick) const;
+	const FFrameBasedTempoPoint* GetTempoPointForFrame(double DecimalFrame) const;
+	const FFrameBasedTempoPoint* GetTempoPointForFrameTime(const FFrameTime& Frame) const;
+	const FFrameBasedTempoPoint* GetTempoPointForSeconds(double Seconds) const;
+
+	FSubdivisionPoint MakeSubdivisionPoint(const FMusicalTime& MusicalTime) const;
+
+	FString MakeBarBeatString(double DecimalFrame, const FFrameBasedTimeSignaturePoint* TimeSignature, const FFrameBasedTempoPoint* Tempo, FrameBasedMusicMap::EBarBeatFormat BarBeatFormat) const;
+};
+
+UINTERFACE(MinimalAPI)
+class UMusicMapSource : public UInterface
+{
+	GENERATED_BODY()
+};
+
+struct FMarkerProviderEntry
+{
+	FString Label;
+	FFrameNumber FrameNumber;
+
+	FMarkerProviderEntry(const FString& InLabel, FFrameNumber InFrameNumber)
+		: Label(InLabel)
+		, FrameNumber(InFrameNumber)
+	{}
+};
+
+USTRUCT()
+struct FMarkerProviderResults
+{
+	GENERATED_BODY()
+
+	struct FChannel
+	{
+		FString Name;
+		TArray<FMarkerProviderEntry> Markers;
+
+		FChannel(const FString& InName)
+			: Name(InName)
+		{}
+	};
+
+	bool IsEmpty() const { return Channels.IsEmpty(); }
+
+	TArray<FChannel> Channels;
+};
+
+class MUSICENVIRONMENT_API IMusicMapSource
+{
+public:
+
+	GENERATED_BODY()
+
+	UFUNCTION()
+	virtual void CreateFrameBasedMusicMap(UFrameBasedMusicMap* Map) const { return; }
+
+	UFUNCTION()
+	virtual FMarkerProviderResults GatherMarkers(const UFrameBasedMusicMap* Map) const { return FMarkerProviderResults(); }
+
+};
+
diff --git a/Engine/Plugins/Runtime/MusicEnvironment/Source/MusicEnvironment/Public/MusicClockSourceManager.h b/Engine/Plugins/Runtime/MusicEnvironment/Source/MusicEnvironment/Public/MusicClockSourceManager.h
new file mode 100644
--- /dev/null
+++ b/Engine/Plugins/Runtime/MusicEnvironment/Source/MusicEnvironment/Public/MusicClockSourceManager.h
@@ -0,0 +1,51 @@
+// Copyright Epic Games, Inc. All Rights Reserved.
+
+#pragma once
+
+#include "GameplayTagContainer.h"
+#include "MusicEnvironmentClockSource.h"
+#include "UObject/WeakInterfacePtr.h"
+
+#include "MusicClockSourceManager.generated.h"
+
+UCLASS(BlueprintType)
+class MUSICENVIRONMENT_API UMusicClockSourceManager : public UObject
+{
+	GENERATED_BODY()
+public:
+	
+	UFUNCTION(BlueprintCallable, Category="Music|Clock")
+	TScriptInterface<IMusicEnvironmentClockSource> FindClock(const FGameplayTag& Tag, bool bExactMatch = true);
+
+	UFUNCTION(BlueprintCallable, Category="Music|Clock")
+	void AddTaggedClock(const FGameplayTag& Tag, TScriptInterface<IMusicEnvironmentClockSource> InClock);
+	void AddTaggedClock(const FGameplayTag& Tag, UObject* InClock);
+
+	UFUNCTION(BlueprintCallable, Category="Music|Clock")
+	void RemoveTaggedClock(TScriptInterface<IMusicEnvironmentClockSource> InClock);
+	void RemoveTaggedClock(UObject* InClock);
+
+	UFUNCTION(BlueprintCallable, Category="Music|Clock")
+	void RemoveClockWithTag(const FGameplayTag& Tag);
+
+	UFUNCTION(BlueprintCallable, Category="Music|Clock")
+	TScriptInterface<IMusicEnvironmentClockSource> GetGlobalMusicClockAuthority();
+
+	UFUNCTION(BlueprintCallable, Category = "Music|Clock")
+	bool PushGlobalMusicClockAuthority(TScriptInterface<IMusicEnvironmentClockSource> InClock);
+	bool PushGlobalMusicClockAuthority(UObject* InClock);
+
+	UFUNCTION(BlueprintCallable, Category = "Music|Clock")
+	void RemoveGlobalClockAuthority(TScriptInterface<IMusicEnvironmentClockSource> InClock);
+	void RemoveGlobalClockAuthority(UObject* InClock);
+
+	UFUNCTION(BlueprintCallable, Category = "Music|Clock")
+	void PopMusicClockAuthority();
+
+private:
+	TMap<FGameplayTag, TWeakInterfacePtr<IMusicEnvironmentClockSource>> TaggedClocks;
+	TArray<TWeakInterfacePtr<IMusicEnvironmentClockSource>> GlobalClockSourceStack;
+
+	void CleanMap();
+	void CleanStack();
+};
diff --git a/Engine/Plugins/Runtime/MusicEnvironment/Source/MusicEnvironment/Public/MusicEnvironmentClockSource.h b/Engine/Plugins/Runtime/MusicEnvironment/Source/MusicEnvironment/Public/MusicEnvironmentClockSource.h
new file mode 100644
--- /dev/null
+++ b/Engine/Plugins/Runtime/MusicEnvironment/Source/MusicEnvironment/Public/MusicEnvironmentClockSource.h
@@ -0,0 +1,27 @@
+// Copyright Epic Games, Inc. All Rights Reserved.
+
+#pragma once
+
+#include "CoreMinimal.h"
+#include "Misc/MusicalTime.h"
+#include "UObject/Interface.h"
+
+#include "MusicEnvironmentClockSource.generated.h"
+
+// This class does not need to be modified.
+UINTERFACE(BlueprintType, meta = (CannotImplementInterfaceInBlueprint))
+class UMusicEnvironmentClockSource : public UInterface
+{
+	GENERATED_BODY()
+};
+
+/**
+ * 
+ */
+class MUSICENVIRONMENT_API IMusicEnvironmentClockSource
+{
+	GENERATED_BODY()
+
+public:
+	virtual FMusicalTime GetMusicalTime() const = 0;
+};
diff --git a/Engine/Plugins/Runtime/MusicEnvironment/Source/MusicEnvironment/Public/MusicEnvironmentMetronome.h b/Engine/Plugins/Runtime/MusicEnvironment/Source/MusicEnvironment/Public/MusicEnvironmentMetronome.h
new file mode 100644
--- /dev/null
+++ b/Engine/Plugins/Runtime/MusicEnvironment/Source/MusicEnvironment/Public/MusicEnvironmentMetronome.h
@@ -0,0 +1,47 @@
+// Copyright Epic Games, Inc. All Rights Reserved.
+
+#pragma once
+
+#include "CoreMinimal.h"
+#include "FrameBasedMusicMap.h"
+#include "MusicEnvironmentClockSource.h"
+#include "UObject/Interface.h"
+
+#include "MusicEnvironmentMetronome.generated.h"
+
+// This class does not need to be modified.
+UINTERFACE()
+class UMusicEnvironmentMetronome : public UInterface
+{
+	GENERATED_BODY()
+};
+
+/**
+ * A music environment system that can spawn metronomes will have to return an 
+ * instance of a UObject that implements this interface from its "MovieSceneMetronomeSpawner". 
+ */
+class MUSICENVIRONMENT_API IMusicEnvironmentMetronome : public IMusicEnvironmentClockSource
+{
+	GENERATED_BODY()
+
+	// Add interface functions to this class. This is the class that will be inherited to implement this interface.
+public:
+	void SetMusicMap(UFrameBasedMusicMap* InMusicMap);
+
+	virtual bool Initialize(UWorld* InWorld) = 0;
+	virtual void Tick(float DeltaSeconds) = 0;
+	virtual void Start(double FromSeconds = 0.0) = 0;
+	virtual void Seek(double ToSeconds) = 0;
+	virtual void Stop() = 0;
+	virtual void Pause() = 0;
+	virtual void Resume() = 0;
+
+	bool SetTempo(float Bpm);
+	virtual float GetCurrentTempo() = 0;
+	virtual double GetCurrentPositionSeconds() = 0;
+	
+protected:
+	virtual void OnMusicMapSet() = 0;
+	virtual bool OnSetTempo(const float Bpm) = 0;
+	TStrongObjectPtr<UFrameBasedMusicMap> MusicMap;
+};
diff --git a/Engine/Plugins/Runtime/MusicEnvironment/Source/MusicEnvironment/Public/MusicEnvironmentModule.h b/Engine/Plugins/Runtime/MusicEnvironment/Source/MusicEnvironment/Public/MusicEnvironmentModule.h
new file mode 100644
--- /dev/null
+++ b/Engine/Plugins/Runtime/MusicEnvironment/Source/MusicEnvironment/Public/MusicEnvironmentModule.h
@@ -0,0 +1,21 @@
+// Copyright Epic Games, Inc. All Rights Reserved.
+
+#pragma once
+
+#include "CoreMinimal.h"
+#include "Modules/ModuleInterface.h"
+#include "Modules/ModuleManager.h"
+
+/**
+ * The public interface of the MusicEnvironment module
+ */
+class FMusicEnvironmentModule
+	: public IModuleInterface
+{
+public:
+
+	static inline FMusicEnvironmentModule& Get()
+	{
+		return FModuleManager::LoadModuleChecked<FMusicEnvironmentModule>("MusicEnvironment");
+	}
+};
diff --git a/Engine/Plugins/Runtime/MusicEnvironment/Source/MusicEnvironment/Public/MusicEnvironmentSubsystem.h b/Engine/Plugins/Runtime/MusicEnvironment/Source/MusicEnvironment/Public/MusicEnvironmentSubsystem.h
new file mode 100644
--- /dev/null
+++ b/Engine/Plugins/Runtime/MusicEnvironment/Source/MusicEnvironment/Public/MusicEnvironmentSubsystem.h
@@ -0,0 +1,44 @@
+// Copyright Epic Games, Inc. All Rights Reserved.
+
+#pragma once
+
+#include "CoreMinimal.h"
+#include "Subsystems/EngineSubsystem.h"
+#include "UObject/Class.h"
+#include "UObject/ObjectMacros.h"
+#include "UObject/UObjectGlobals.h"
+#include "UObject/WeakInterfacePtr.h"
+
+#include "MusicEnvironmentSubsystem.generated.h"
+
+// Forward Declarations
+class IMusicEnvironmentMetronome;
+class UMusicClockSourceManager;
+
+UCLASS(BlueprintType, Category = "Music", DisplayName = "MusicEnvironmentSubsystem")
+class MUSICENVIRONMENT_API UMusicEnvironmentSubsystem : public UEngineSubsystem
+{
+	GENERATED_BODY()
+public:
+	static UMusicEnvironmentSubsystem& Get();
+	
+	UFUNCTION(BlueprintCallable, Category = "Music|Clock")
+	UMusicClockSourceManager* GetClockSourceManager();
+
+	bool SetMetronomeClass(const UClass* InMetronomeType);
+
+	UFUNCTION(BlueprintCallable, Category = "Music|Clock")
+	bool CanSpawnMetronome() const;
+
+	UFUNCTION(BlueprintCallable, Category = "Music|Clock")
+	virtual TScriptInterface<IMusicEnvironmentMetronome> SpawnMetronome(UObject* Outer);
+
+private:
+	UPROPERTY(Transient)
+	TObjectPtr<UMusicClockSourceManager> ClockSourceManager;
+
+	UPROPERTY(Transient)
+	TSoftClassPtr<UObject> MetronomeType;
+};
+
+MUSICENVIRONMENT_API DECLARE_LOG_CATEGORY_EXTERN(LogMusicEnvironment, Log, All);
diff --git a/Engine/Source/Runtime/AudioMixer/AudioMixer.Build.cs b/Engine/Source/Runtime/AudioMixer/AudioMixer.Build.cs
--- a/Engine/Source/Runtime/AudioMixer/AudioMixer.Build.cs
+++ b/Engine/Source/Runtime/AudioMixer/AudioMixer.Build.cs
@@ -32,7 +32,6 @@ namespace UnrealBuildTool.Rules
 					"AudioExtensions",
 					"AudioLinkCore",
 					"HeadMountedDisplay",
-					"MusicEnvironment",
 					"TraceLog"
 				}
 			);
diff --git a/Engine/Source/Runtime/AudioMixer/Private/AudioMixerModule.cpp b/Engine/Source/Runtime/AudioMixer/Private/AudioMixerModule.cpp
--- a/Engine/Source/Runtime/AudioMixer/Private/AudioMixerModule.cpp
+++ b/Engine/Source/Runtime/AudioMixer/Private/AudioMixerModule.cpp
@@ -10,7 +10,6 @@ public:
 	virtual void StartupModule() override
 	{
 		FModuleManager::Get().LoadModuleChecked(TEXT("AudioMixerCore"));
-		FModuleManager::Get().LoadModuleChecked(TEXT("MusicEnvironment"));
 		FModuleManager::Get().LoadModuleChecked(TEXT("SignalProcessing"));
 	}
 };
diff --git a/Engine/Source/Runtime/MusicEnvironment/MusicEnvironment.Build.cs b/Engine/Source/Runtime/MusicEnvironment/MusicEnvironment.Build.cs
deleted file mode 100644
diff --git a/Engine/Source/Runtime/MusicEnvironment/Private/FrameBasedMusicMap.cpp b/Engine/Source/Runtime/MusicEnvironment/Private/FrameBasedMusicMap.cpp
deleted file mode 100644
diff --git a/Engine/Source/Runtime/MusicEnvironment/Private/MusicClockSourceManager.cpp b/Engine/Source/Runtime/MusicEnvironment/Private/MusicClockSourceManager.cpp
deleted file mode 100644
diff --git a/Engine/Source/Runtime/MusicEnvironment/Private/MusicEnvironmentMetronome.cpp b/Engine/Source/Runtime/MusicEnvironment/Private/MusicEnvironmentMetronome.cpp
deleted file mode 100644
diff --git a/Engine/Source/Runtime/MusicEnvironment/Private/MusicEnvironmentModule.cpp b/Engine/Source/Runtime/MusicEnvironment/Private/MusicEnvironmentModule.cpp
deleted file mode 100644
diff --git a/Engine/Source/Runtime/MusicEnvironment/Private/MusicEvironmentSubsystem.cpp b/Engine/Source/Runtime/MusicEnvironment/Private/MusicEvironmentSubsystem.cpp
deleted file mode 100644
diff --git a/Engine/Source/Runtime/MusicEnvironment/Public/FrameBasedMusicMap.h b/Engine/Source/Runtime/MusicEnvironment/Public/FrameBasedMusicMap.h
deleted file mode 100644
diff --git a/Engine/Source/Runtime/MusicEnvironment/Public/MusicClockSourceManager.h b/Engine/Source/Runtime/MusicEnvironment/Public/MusicClockSourceManager.h
deleted file mode 100644
diff --git a/Engine/Source/Runtime/MusicEnvironment/Public/MusicEnvironmentClockSource.h b/Engine/Source/Runtime/MusicEnvironment/Public/MusicEnvironmentClockSource.h
deleted file mode 100644
diff --git a/Engine/Source/Runtime/MusicEnvironment/Public/MusicEnvironmentMetronome.h b/Engine/Source/Runtime/MusicEnvironment/Public/MusicEnvironmentMetronome.h
deleted file mode 100644
diff --git a/Engine/Source/Runtime/MusicEnvironment/Public/MusicEnvironmentModule.h b/Engine/Source/Runtime/MusicEnvironment/Public/MusicEnvironmentModule.h
deleted file mode 100644
diff --git a/Engine/Source/Runtime/MusicEnvironment/Public/MusicEnvironmentSubsystem.h b/Engine/Source/Runtime/MusicEnvironment/Public/MusicEnvironmentSubsystem.h
deleted file mode 100644
