Move MusicEnvironment module from Runtime "Module" to Runtime Plugin

Source/Runtime/MusicEnvironment -> Plugins/Runtime/MusicEnvironment/Source/MusicEnvironment

#rnx

diff --git a/Engine/Plugins/Runtime/Harmonix/Harmonix.uplugin b/Engine/Plugins/Runtime/Harmonix/Harmonix.uplugin
--- a/Engine/Plugins/Runtime/Harmonix/Harmonix.uplugin
+++ b/Engine/Plugins/Runtime/Harmonix/Harmonix.uplugin
@@ -77,6 +77,10 @@
 			"Name": "Metasound",
 			"Enabled": true
 		},
+		{
+			"Name": "MusicEnvironment",
+			"Enabled": true
+		},
 		{
 			"Name": "EditorScriptingUtilities",
 			"Enabled":  true
diff --git a/Engine/Plugins/Runtime/MusicEnvironment/MusicEnvironment.uplugin b/Engine/Plugins/Runtime/MusicEnvironment/MusicEnvironment.uplugin
new file mode 100644
--- /dev/null
+++ b/Engine/Plugins/Runtime/MusicEnvironment/MusicEnvironment.uplugin
@@ -0,0 +1,27 @@
+﻿{
+	"FileVersion": 3,
+	"Version": 1,
+	"VersionName": "1.0",
+	"FriendlyName": "Music Environment",
+	"Description": "A Project-Wide source of musical information (musically synchronized clocks, events, etc.)",
+	"Category": "Audio",
+	"CreatedBy": "Epic Games, Inc. - Harmonix GenTech",
+	"CreatedByURL": "",
+	"DocsURL": "",
+	"MarketplaceURL": "",
+	"SupportURL": "",
+	"EnabledByDefault": false,
+	"CanContainContent": false,
+	"IsBetaVersion": true,
+	"IsExperimentalVersion": true,
+	"Installed": false,
+	"Modules": [
+		{
+			"Name": "MusicEnvironment",
+			"Type": "Runtime",
+			"LoadingPhase": "PreDefault"
+		}
+	],
+	"Plugins": [
+	]
+}
\ No newline at end of file
diff --git a/Engine/Plugins/Runtime/MusicEnvironment/Source/MusicEnvironment/MusicEnvironment.Build.cs b/Engine/Plugins/Runtime/MusicEnvironment/Source/MusicEnvironment/MusicEnvironment.Build.cs
new file mode 100644
--- /dev/null
+++ b/Engine/Plugins/Runtime/MusicEnvironment/Source/MusicEnvironment/MusicEnvironment.Build.cs
@@ -0,0 +1,25 @@
+// Copyright Epic Games, Inc. All Rights Reserved.
+
+using UnrealBuildTool;
+
+public class MusicEnvironment : ModuleRules
+{
+	public MusicEnvironment(ReadOnlyTargetRules Target) : base(Target)
+	{
+        PrivateIncludePathModuleNames.AddRange(
+            new string[] {
+			}
+        );
+
+        PublicDependencyModuleNames.AddRange(
+			new string[] {
+				"Core",
+				"CoreUObject",
+                "Engine",
+				"TimeManagement",
+				"GameplayTags",
+				"MovieScene",
+			}
+		);
+	}
+}
diff --git a/Engine/Plugins/Runtime/MusicEnvironment/Source/MusicEnvironment/Private/FrameBasedMusicMap.cpp b/Engine/Plugins/Runtime/MusicEnvironment/Source/MusicEnvironment/Private/FrameBasedMusicMap.cpp
new file mode 100644
--- /dev/null
+++ b/Engine/Plugins/Runtime/MusicEnvironment/Source/MusicEnvironment/Private/FrameBasedMusicMap.cpp
@@ -0,0 +1,951 @@
+// Copyright Epic Games, Inc. All Rights Reserved.
+
+#include "FrameBasedMusicMap.h"
+#include "Algo/FindLast.h"
+
+DEFINE_LOG_CATEGORY_STATIC(LogMusicMap, Log, All);
+#define LOCTEXT_NAMESPACE "FrameBasedMusicMap"
+
+namespace FrameBasedMusicMap
+{
+	FText GetQuantizationName(int32 Quantization)
+	{
+		switch (Quantization)
+		{
+		case ETripletQuantizationTicks::ThirtysecondNoteTriplet:	return LOCTEXT("ThirtySecondNoteTriplet", "1/32 Note Triplet");
+		case ETripletQuantizationTicks::SixteenthNoteTriplet:		return LOCTEXT("SixteenthNoteTriplet", "1/16 Note Triplet");
+		case ETripletQuantizationTicks::EighthNoteTriplet:			return LOCTEXT("EighthNoteTriplet", "1/8 Note Triplet");
+		case ETripletQuantizationTicks::QuarterNoteTriplet:			return LOCTEXT("QuarterNoteTriplet", "1/4 Note Triplet");
+		case ETripletQuantizationTicks::HalfNoteTriplet:			return LOCTEXT("HalfNoteTriplet", "1/2 Note Triplet");
+		case ETripletQuantizationTicks::WholeNoteTriplet:			return LOCTEXT("WholeNoteTriplet", "Whole Note Triplet");
+		
+		case EWholeQuantizationTicks::ThirtysecondNote:	return LOCTEXT("ThirtySecondNote", "1/32 Note");
+		case EWholeQuantizationTicks::SixteenthNote:	return LOCTEXT("SixteenthNote", "1/16 Note");
+		case EWholeQuantizationTicks::EighthNote:		return LOCTEXT("EighthNote", "1/8 Note");
+		case EWholeQuantizationTicks::QuarterNote:		return LOCTEXT("QuarterNote", "1/4 Note");
+		case EWholeQuantizationTicks::HalfNote:			return LOCTEXT("HalfNote", "1/2 Note");
+		case EWholeQuantizationTicks::WholeNote:		return LOCTEXT("WholeNote", "Whole Note");
+		case EWholeQuantizationTicks::Bar:				return LOCTEXT("Bar", "Bar");
+		case EWholeQuantizationTicks::Beat:				return LOCTEXT("Beat", "Beat");
+
+
+		case EDottedQuantizationTicks::DottedThirtysecondNote:	return LOCTEXT("DottedThirtySecondNote", "Dotted 1/32 Note");
+		case EDottedQuantizationTicks::DottedSixteenthNote:		return LOCTEXT("DottedSixteenthNote", "Dotted 1/16 Note");
+		case EDottedQuantizationTicks::DottedEighthNote:		return LOCTEXT("DottedEighthNote", "Dotted 1/8 Note");
+		case EDottedQuantizationTicks::DottedQuarterNote:		return LOCTEXT("DottedQuarterNote", "Dotted 1/4 Note");
+		case EDottedQuantizationTicks::DottedHalfNote:			return LOCTEXT("DottedHalfNote", "Dotted Half Note");
+		case EDottedQuantizationTicks::DottedWholeNote:			return LOCTEXT("DottedWholeNote", "Dotted Whole Note");
+		}
+		return LOCTEXT("Unknown", "Unrecognized Quantization Factor");
+	}
+
+	EQuartzCommandQuantization QuantiazationIntervalToQuartz(int32 Quantization)
+	{
+		switch (Quantization)
+		{
+		case ETripletQuantizationTicks::ThirtysecondNoteTriplet:	return EQuartzCommandQuantization::None; // quartz doesn't go this small.
+		case ETripletQuantizationTicks::SixteenthNoteTriplet:		return EQuartzCommandQuantization::SixteenthNoteTriplet;
+		case ETripletQuantizationTicks::EighthNoteTriplet:			return EQuartzCommandQuantization::EighthNoteTriplet;
+		case ETripletQuantizationTicks::QuarterNoteTriplet:			return EQuartzCommandQuantization::QuarterNoteTriplet;
+		case ETripletQuantizationTicks::HalfNoteTriplet:			return EQuartzCommandQuantization::HalfNoteTriplet;
+		case ETripletQuantizationTicks::WholeNoteTriplet:			return EQuartzCommandQuantization::None; // quartz doesn't support this division.
+
+		case EWholeQuantizationTicks::ThirtysecondNote:	return EQuartzCommandQuantization::ThirtySecondNote;
+		case EWholeQuantizationTicks::SixteenthNote:	return EQuartzCommandQuantization::SixteenthNote;
+		case EWholeQuantizationTicks::EighthNote:		return EQuartzCommandQuantization::EighthNote;
+		case EWholeQuantizationTicks::QuarterNote:		return EQuartzCommandQuantization::QuarterNote;
+		case EWholeQuantizationTicks::HalfNote:			return EQuartzCommandQuantization::HalfNote;
+		case EWholeQuantizationTicks::WholeNote:		return EQuartzCommandQuantization::WholeNote;
+		case EWholeQuantizationTicks::Bar:				return EQuartzCommandQuantization::Bar;
+		case EWholeQuantizationTicks::Beat:				return EQuartzCommandQuantization::Beat;
+
+
+		case EDottedQuantizationTicks::DottedThirtysecondNote:	return EQuartzCommandQuantization::None; // quartz doesn't support this division.
+		case EDottedQuantizationTicks::DottedSixteenthNote:		return EQuartzCommandQuantization::DottedSixteenthNote;
+		case EDottedQuantizationTicks::DottedEighthNote:		return EQuartzCommandQuantization::DottedEighthNote;
+		case EDottedQuantizationTicks::DottedQuarterNote:		return EQuartzCommandQuantization::DottedQuarterNote;
+		case EDottedQuantizationTicks::DottedHalfNote:			return EQuartzCommandQuantization::DottedHalfNote;
+		case EDottedQuantizationTicks::DottedWholeNote:			return EQuartzCommandQuantization::DottedWholeNote;
+		}
+		return EQuartzCommandQuantization::None;
+	}
+
+	TArray<TSharedPtr<int32>> QuantizationOptionsSource = {
+		MakeShared<int32>(EWholeQuantizationTicks::Bar),
+		MakeShared<int32>(EWholeQuantizationTicks::Beat),
+		MakeShared<int32>(EWholeQuantizationTicks::WholeNote),
+		MakeShared<int32>(EWholeQuantizationTicks::HalfNote),
+		MakeShared<int32>(EWholeQuantizationTicks::QuarterNote),
+		MakeShared<int32>(EWholeQuantizationTicks::EighthNote),
+		MakeShared<int32>(EWholeQuantizationTicks::SixteenthNote),
+		MakeShared<int32>(EWholeQuantizationTicks::ThirtysecondNote),
+		MakeShared<int32>(ETripletQuantizationTicks::WholeNoteTriplet),
+		MakeShared<int32>(ETripletQuantizationTicks::HalfNoteTriplet),
+		MakeShared<int32>(ETripletQuantizationTicks::QuarterNoteTriplet),
+		MakeShared<int32>(ETripletQuantizationTicks::EighthNoteTriplet),
+		MakeShared<int32>(ETripletQuantizationTicks::SixteenthNoteTriplet),
+		MakeShared<int32>(ETripletQuantizationTicks::ThirtysecondNoteTriplet),
+		MakeShared<int32>(EDottedQuantizationTicks::DottedWholeNote),
+		MakeShared<int32>(EDottedQuantizationTicks::DottedHalfNote),
+		MakeShared<int32>(EDottedQuantizationTicks::DottedQuarterNote),
+		MakeShared<int32>(EDottedQuantizationTicks::DottedEighthNote),
+		MakeShared<int32>(EDottedQuantizationTicks::DottedSixteenthNote),
+		MakeShared<int32>(EDottedQuantizationTicks::DottedThirtysecondNote)
+	};
+}
+
+void UFrameBasedMusicMap::SetFrameResolution(const FFrameRate& NewFrameResolution)
+{
+	check(NewFrameResolution.IsValid());
+
+	if (NewFrameResolution == FrameResolution)
+	{
+		return;
+	}
+
+	if (FrameResolution.IsValid() && !IsEmpty() && !IsSimple())
+	{
+		// must be changing sequence's resolution...
+		for (FFrameBasedTimeSignaturePoint& Bar : BarMap)
+		{
+			Bar.OnFrame = FFrameRate::TransformTime(FFrameTime::FromDecimal(Bar.OnFrame), FrameResolution, NewFrameResolution).AsDecimal();
+		}
+	}
+
+	FrameResolution = NewFrameResolution;
+}
+
+void UFrameBasedMusicMap::Init(float InBpm, const FFrameBasedTimeSignature& InTimeSignature)
+{
+	check(FrameResolution.IsValid());
+	Clear();
+	TempoMap.Add(FFrameBasedTempoPoint(InBpm, 0, 0.0f));
+	BarMap.Add(FFrameBasedTimeSignaturePoint(InTimeSignature, 0, 0, 0.0f));
+}
+
+void UFrameBasedMusicMap::SetInitialTempo(float Bpm)
+{
+	check(!TempoMap.IsEmpty());
+	TempoMap[0].SetBpm(Bpm);
+}
+
+float UFrameBasedMusicMap::GetInitialTempo() const
+{
+	return TempoMap.IsEmpty() ? 0.0f : TempoMap[0].Bpm();
+}
+
+void UFrameBasedMusicMap::SetInitialTimeSignature(const FFrameBasedTimeSignature& TimeSignature)
+{
+	check(!BarMap.IsEmpty());
+	BarMap[0].TimeSignature = TimeSignature;
+}
+
+FFrameBasedTimeSignature UFrameBasedMusicMap::GetInitialTimeSignature() const
+{
+	return BarMap.IsEmpty() ? FFrameBasedTimeSignature() : BarMap[0].TimeSignature;
+}
+
+double UFrameBasedMusicMap::GetFrame(const FMusicalTime& MusicalTime) const
+{
+	return FrameResolution.AsFrameTime(GetSeconds(MusicalTime)).AsDecimal();
+}
+
+double UFrameBasedMusicMap::GetFrame(int32 AtBar) const
+{
+	return FrameResolution.AsFrameTime(GetSeconds(AtBar)).AsDecimal();
+}
+
+FFrameTime UFrameBasedMusicMap::GetFrameTime(const FMusicalTime& MusicalTime) const
+{
+	return FrameResolution.AsFrameTime(GetSeconds(MusicalTime));
+}
+
+FFrameTime UFrameBasedMusicMap::GetFrameTime(int32 AtBar, float AtBeat) const
+{
+	return FrameResolution.AsFrameTime(GetSeconds(AtBar, AtBeat));
+}
+
+FFrameTime UFrameBasedMusicMap::GetFrameTime(double FractionalBar) const
+{
+	if (IsEmpty())
+	{
+		return FFrameTime();
+	}
+
+	int32 Bar = FMath::FloorToInt32(FractionalBar);
+	float Beat =static_cast<float>(static_cast<double>(GetBeatsInBar(Bar)) * FMath::Fractional(FractionalBar));
+	return GetFrameTime(Bar, Beat);	
+}
+
+FFrameTime UFrameBasedMusicMap::GetFrameTime(int32 AtTick) const
+{
+	if (IsEmpty())
+	{
+		return FFrameTime();
+	}
+
+	FMusicalTime MusicalTime = GetMusicalTime(AtTick);
+	return GetFrameTime(MusicalTime);
+}
+
+double UFrameBasedMusicMap::GetSeconds(const FMusicalTime& MusicalTime) const
+{
+	if (IsEmpty())
+	{
+		return 0.0;
+	}
+
+	check(MusicalTime.IsValid());
+
+	const FFrameBasedTimeSignaturePoint* TimeSignatureAtFrame = GetTimeSignaturePointForBar(MusicalTime.Bar);
+
+	int32 Tick = TimeSignatureAtFrame->MusicalTimeToTick(MusicalTime);
+
+	const FFrameBasedTempoPoint* TempoAtFrame = GetTempoPointForTick(Tick);
+
+	return TempoAtFrame->TickToSeconds(Tick);
+}
+
+double UFrameBasedMusicMap::GetSeconds(int32 AtBar, float AtBeat) const
+{
+	if (IsEmpty())
+	{
+		return 0.0;
+	}
+
+	const FFrameBasedTimeSignaturePoint* TimeSignatureAtFrame = GetTimeSignaturePointForBar(AtBar);
+
+	int32 Tick = TimeSignatureAtFrame->BarBeatToTick(AtBar, AtBeat);
+
+	const FFrameBasedTempoPoint* TempoAtFrame = GetTempoPointForTick(Tick);
+
+	return TempoAtFrame->TickToSeconds(Tick);
+}
+
+FString UFrameBasedMusicMap::MakeBarBeatStringForFrame(double DecimalFrame, FrameBasedMusicMap::EBarBeatFormat BarBeatFormat) const
+{
+	if (IsEmpty())
+	{
+		return FString(TEXT("0.0.00"));
+	}
+
+	const FFrameBasedTimeSignaturePoint* TimeSignatureAtFrame = GetTimeSignaturePointForFrame(DecimalFrame);
+	const FFrameBasedTempoPoint* TempoAtFrame = GetTempoPointForFrame(DecimalFrame);
+
+	return MakeBarBeatString(DecimalFrame, TimeSignatureAtFrame, TempoAtFrame, BarBeatFormat);
+}
+
+FString UFrameBasedMusicMap::MakeBarBeatStringForFrameTime(const FFrameTime& Frame, FrameBasedMusicMap::EBarBeatFormat BarBeatFormat) const
+{
+	return MakeBarBeatStringForFrame(Frame.AsDecimal(), BarBeatFormat);
+}
+
+FString UFrameBasedMusicMap::MakeBarBeatStringForMusicalTime(const FMusicalTime& MusicalTime, FrameBasedMusicMap::EBarBeatFormat BarBeatFormat) const
+{
+	FFrameTime FrameTime = GetFrameTime(MusicalTime);
+	return MakeBarBeatStringForFrameTime(FrameTime, BarBeatFormat);
+}
+
+FString UFrameBasedMusicMap::MakeBarBeatStringForSeconds(double Seconds, FrameBasedMusicMap::EBarBeatFormat BarBeatFormat) const
+{
+	return MakeBarBeatStringForFrameTime(FrameResolution.AsFrameTime(Seconds), BarBeatFormat);
+}
+
+float UFrameBasedMusicMap::GetTempoForFrame(double DecimalFrame) const
+{
+	return GetTempoForFrameTime(FFrameTime::FromDecimal(DecimalFrame));
+}
+
+float UFrameBasedMusicMap::GetTempoForFrameTime(const FFrameTime& FrameTime) const
+{
+	return GetTempoForSeconds(FrameResolution.AsSeconds(FrameTime));
+}
+
+float UFrameBasedMusicMap::GetTempoForSeconds(double Seconds) const
+{
+	if (TempoMap.IsEmpty())
+	{
+		return 0.0f;
+	}
+
+	if (Seconds <= 0.0)
+	{
+		return TempoMap[0].Bpm();
+	}
+
+	const FFrameBasedTempoPoint* TempoPoint = GetTempoPointForSeconds(Seconds);
+
+	return TempoPoint->Bpm();
+}
+
+FFrameBasedTimeSignature UFrameBasedMusicMap::GetTimeSignatureForFrame(double DecimalFrame) const
+{
+	const FFrameBasedTimeSignaturePoint* Point = GetTimeSignaturePointForFrame(DecimalFrame);
+	return Point ? Point->TimeSignature : FFrameBasedTimeSignature();
+}
+
+FFrameBasedTimeSignature UFrameBasedMusicMap::GetTimeSignatureForFrameTime(const FFrameTime& FrameTime) const
+{
+	const FFrameBasedTimeSignaturePoint* Point = GetTimeSignaturePointForFrameTime(FrameTime);
+	return Point ? Point->TimeSignature : FFrameBasedTimeSignature();
+}
+
+FFrameBasedTimeSignature UFrameBasedMusicMap::GetTimeSignatureForSeconds(double Seconds) const
+{
+	const FFrameBasedTimeSignaturePoint* Point = GetTimeSignaturePointForSeconds(Seconds);
+	return Point ? Point->TimeSignature : FFrameBasedTimeSignature();
+}
+
+int32 UFrameBasedMusicMap::GetBeatsInBar(int32 Bar) const
+{
+	if (BarMap.IsEmpty())
+	{
+		return 4;
+	}
+	
+	const FFrameBasedTimeSignaturePoint* TimeSignaturePoint = GetTimeSignaturePointForBar(Bar);
+	check(TimeSignaturePoint);
+	return TimeSignaturePoint->TimeSignature.Numerator;
+}
+
+int32 UFrameBasedMusicMap::GetTicksInBar(int32 Bar) const
+{
+	if (BarMap.IsEmpty())
+	{
+		return 4 * MusicalTime::TicksPerQuarterNote;
+	}
+	
+	const FFrameBasedTimeSignaturePoint* TimeSignaturePoint = GetTimeSignaturePointForBar(Bar);
+	check(TimeSignaturePoint);
+	return TimeSignaturePoint->TicksPerBar();	
+}
+
+const FFrameBasedTimeSignaturePoint* UFrameBasedMusicMap::GetTimeSignaturePointForTick(int32 Tick) const
+{
+	if (BarMap.IsEmpty())
+	{
+		return nullptr;
+	}
+
+	if (Tick <= 0)
+	{
+		return &BarMap[0];
+	}
+
+	const FFrameBasedTimeSignaturePoint* TimeSignatureAtTick = Algo::FindLastByPredicate(BarMap, [=](FFrameBasedTimeSignaturePoint const& Point) { return Tick >= Point.OnTick; });
+	return TimeSignatureAtTick ? TimeSignatureAtTick : &BarMap[0];
+}
+
+const FFrameBasedTimeSignaturePoint* UFrameBasedMusicMap::GetTimeSignaturePointForBar(int32 Bar) const
+{
+	if (BarMap.IsEmpty())
+	{
+		return nullptr;
+	}
+
+	if (Bar <= 0)
+	{
+		return &BarMap[0];
+	}
+
+	const FFrameBasedTimeSignaturePoint* TimeSignatureAtBar = Algo::FindLastByPredicate(BarMap, [=](FFrameBasedTimeSignaturePoint const& Point) { return Bar >= Point.OnBar; });
+	return TimeSignatureAtBar ? TimeSignatureAtBar : &BarMap[0];
+}
+
+const FFrameBasedTimeSignaturePoint* UFrameBasedMusicMap::GetTimeSignaturePointForFrame(double DecimalFrame) const
+{
+	if (BarMap.IsEmpty())
+	{
+		return nullptr;
+	}
+
+	if (DecimalFrame <= 0.0)
+	{
+		return &BarMap[0];
+	}
+
+	const FFrameBasedTimeSignaturePoint* TimeSignatureAtBar = Algo::FindLastByPredicate(BarMap, [=](FFrameBasedTimeSignaturePoint const& Point) { return DecimalFrame >= Point.OnFrame; });
+	return TimeSignatureAtBar ? TimeSignatureAtBar : &BarMap[0];
+}
+
+const FFrameBasedTimeSignaturePoint* UFrameBasedMusicMap::GetTimeSignaturePointForFrameTime(const FFrameTime& Frame) const
+{
+	return GetTimeSignaturePointForFrame(Frame.AsDecimal());
+}
+
+const FFrameBasedTimeSignaturePoint* UFrameBasedMusicMap::GetTimeSignaturePointForSeconds(double Seconds) const
+{
+	if (BarMap.IsEmpty() || TempoMap.IsEmpty())
+	{
+		return nullptr;
+	}
+
+	if (Seconds <= 0.0)
+	{
+		return &BarMap[0];
+	}
+
+	const FFrameBasedTempoPoint* TempoAtTime = GetTempoPointForSeconds(Seconds);
+	double Tick = TempoAtTime->SecondsToTick(Seconds);
+	return GetTimeSignaturePointForTick(Tick);
+}
+
+const FFrameBasedTempoPoint* UFrameBasedMusicMap::GetTempoPointForTick(int32 Tick) const
+{
+	if (TempoMap.IsEmpty())
+	{
+		return nullptr;
+	}
+
+	if (Tick <= 0)
+	{
+		return &TempoMap[0];
+	}
+
+	const FFrameBasedTempoPoint* TempoAtTick = Algo::FindLastByPredicate(TempoMap, [=](FFrameBasedTempoPoint const& Point) { return Tick >= Point.OnTick; });
+	return TempoAtTick ? TempoAtTick : &TempoMap[0];
+}
+
+const FFrameBasedTempoPoint* UFrameBasedMusicMap::GetTempoPointForFrame(double DecimalFrame) const
+{
+	return GetTempoPointForFrameTime(FFrameTime::FromDecimal(DecimalFrame));
+}
+
+const FFrameBasedTempoPoint* UFrameBasedMusicMap::GetTempoPointForFrameTime(const FFrameTime& Frame) const
+{
+	return GetTempoPointForSeconds(FrameResolution.AsSeconds(Frame));
+}
+
+const FFrameBasedTempoPoint* UFrameBasedMusicMap::GetTempoPointForSeconds(double Seconds) const
+{
+	if (TempoMap.IsEmpty())
+	{
+		return nullptr;
+	}
+
+	if (Seconds <= 0.0)
+	{
+		return &TempoMap[0];
+	}
+
+	float Ms = static_cast<float>(Seconds * 1000.0);
+	const FFrameBasedTempoPoint* TempoAtTime = Algo::FindLastByPredicate(TempoMap, [=](FFrameBasedTempoPoint const& Point) { return Ms >= Point.OnMs; });
+	return TempoAtTime ? TempoAtTime : &TempoMap[0];
+}
+
+FFrameTime UFrameBasedMusicMap::RemapFrameTime(FFrameTime SourceFrameTime, const UFrameBasedMusicMap* OldMap) const
+{
+	if (IsEmpty())
+	{
+		return SourceFrameTime;
+	}
+	
+	float Seconds = FrameResolution.AsSeconds(SourceFrameTime);
+	const FFrameBasedTempoPoint* TempoPoint = GetTempoPointForSeconds(Seconds);
+	int32 Tick = TempoPoint->SecondsToTick(Seconds);
+
+	FMusicalTime MusicalTime = OldMap->GetMusicalTime(Tick);
+	const FFrameBasedTimeSignaturePoint* TimeSignaturePoint = GetTimeSignaturePointForBar(MusicalTime.Bar);
+	if (TimeSignaturePoint->TicksPerBar() != MusicalTime.TicksPerBar)
+	{
+		// Time Signature mismatch! Scale bar to bar!
+		MusicalTime.TickInBar = static_cast<int32>( static_cast<float>(MusicalTime.TickInBar * TimeSignaturePoint->TicksPerBar()) / static_cast<float>(MusicalTime.TicksPerBar));
+	}
+	MusicalTime.TicksPerBar = TimeSignaturePoint->TicksPerBar();
+	MusicalTime.TicksPerBeat = TimeSignaturePoint->TicksPerBeat();
+	return GetFrameTime(MusicalTime);
+}
+
+FMusicalTime UFrameBasedMusicMap::BarsToMusicalTime(double BarsFromStart) const
+{
+	FMusicalTime Result;
+	Result.Bar = FMath::FloorToInt32(BarsFromStart);
+	double FractionalBar = BarsFromStart - static_cast<double>(Result.Bar);
+
+	const FFrameBasedTimeSignaturePoint* TimeSignaturePoint = GetTimeSignaturePointForBar(Result.Bar);
+	Result.TicksPerBar = TimeSignaturePoint->TicksPerBar();
+	Result.TicksPerBeat = TimeSignaturePoint->TicksPerBeat();
+	Result.TickInBar = FMath::Clamp(FMath::RoundToInt32( static_cast<double>(Result.TicksPerBar) * FractionalBar), 0, Result.TicksPerBar - 1);
+	return Result;
+}
+
+FMusicalTime UFrameBasedMusicMap::Quantize(const FMusicalTime& MusicalTime, int32 QuantizationInterval, EQuantizeDirection Direction) const
+{
+	int32 QuantizedTick = QuantizeToTick(MusicalTime, QuantizationInterval, Direction);
+
+	const FFrameBasedTimeSignaturePoint* TimeSignaturePoint = GetTimeSignaturePointForTick(QuantizedTick);
+
+	return TimeSignaturePoint->TickToMusicalTime(QuantizedTick);
+}
+
+int32 UFrameBasedMusicMap::QuantizeToTick(const FMusicalTime& MusicalTime, int32 QuantizationInterval, EQuantizeDirection Direction /*= EQuantizeDirection::Nearest*/) const
+{
+	int32 TicksInQuantizationInterval = QuantizationInterval;
+	switch (QuantizationInterval)
+	{
+	case FrameBasedMusicMap::EWholeQuantizationTicks::Bar:
+		TicksInQuantizationInterval = MusicalTime.TicksPerBar;
+		break;
+	case FrameBasedMusicMap::EWholeQuantizationTicks::Beat:
+		TicksInQuantizationInterval = MusicalTime.TicksPerBeat;
+		break;
+	}
+
+	const FFrameBasedTimeSignaturePoint* TimeSignaturePoint = GetTimeSignaturePointForBar(MusicalTime.Bar);
+
+	int32 TotalTicks = TimeSignaturePoint->MusicalTimeToTick(MusicalTime);
+
+	int32 TicksInTimeSignature = TotalTicks - TimeSignaturePoint->OnTick;
+
+	int32 Intervals = 0;
+	switch (Direction)
+	{
+	default:
+	case EQuantizeDirection::Nearest:
+		Intervals = FMath::RoundToInt32(static_cast<float>(TicksInTimeSignature) / static_cast<float>(TicksInQuantizationInterval));
+		break;
+	case EQuantizeDirection::Ceiling:
+		Intervals = FMath::CeilToInt32(static_cast<float>(TicksInTimeSignature) / static_cast<float>(TicksInQuantizationInterval));
+		break;
+	case EQuantizeDirection::Floor:
+		Intervals = FMath::FloorToInt32(static_cast<float>(TicksInTimeSignature) / static_cast<float>(TicksInQuantizationInterval));
+		break;
+	}
+
+	int32 QuantizedTicksInTimeSignature = Intervals * TicksInQuantizationInterval;
+
+	int32 FinalTick = TimeSignaturePoint->OnTick + QuantizedTicksInTimeSignature;
+
+	return FinalTick;
+}
+
+FMusicalTime UFrameBasedMusicMap::Add(const FMusicalTime& MusicalTime, int32 QuantizationInterval) const
+{
+	const FFrameBasedTimeSignaturePoint* OriginalTimeSignature = GetTimeSignaturePointForBar(MusicalTime.Bar);
+	int32 AbsoluteTick = OriginalTimeSignature->MusicalTimeToTick(MusicalTime);
+	if (FMath::Abs(QuantizationInterval) == FrameBasedMusicMap::EWholeQuantizationTicks::Bar)
+	{
+		AbsoluteTick += (QuantizationInterval < 0) ? MusicalTime.TicksPerBar : -MusicalTime.TicksPerBar;
+	}
+	else if (FMath::Abs(QuantizationInterval) == FrameBasedMusicMap::EWholeQuantizationTicks::Beat)
+	{
+		AbsoluteTick += (QuantizationInterval < 0) ? MusicalTime.TicksPerBeat : -MusicalTime.TicksPerBeat;
+	}
+	else
+	{
+		AbsoluteTick += QuantizationInterval;
+	}
+	const FFrameBasedTimeSignaturePoint* NewTimeSignature = GetTimeSignaturePointForTick(AbsoluteTick);
+	return NewTimeSignature->TickToMusicalTime(AbsoluteTick);
+}
+
+void UFrameBasedMusicMap::AddTimeSignature(int32 AtTick, int32 AtBar, int16 Numerator, int16 Denominator)
+{
+	if (ensure(!BarMap.IsEmpty() || AtTick == 0))
+	{
+		if (ensure(AtTick == 0 || BarMap.Last().OnTick < AtTick))
+		{
+			BarMap.Add(FFrameBasedTimeSignaturePoint({Numerator, Denominator}, AtBar, AtTick, GetFrame(AtBar)));
+		}
+	}
+}
+
+void UFrameBasedMusicMap::AddTempo(int32 AtTick, float AtMs, float Bpm)
+{
+	if (ensure(TempoMap.IsEmpty() || AtTick >= TempoMap.Last().OnTick))
+	{
+		float ExpectedMs = TempoMap.IsEmpty() ? 0.0f : static_cast<float>(TempoMap.Last().TickToSeconds(static_cast<double>(AtTick)) * 1000.0);
+		TempoMap.Add(FFrameBasedTempoPoint(Bpm, AtTick, ExpectedMs));
+	}
+}
+
+void UFrameBasedMusicMap::GatherSubdivisonsForTimespan(int32 StartTick, int32 EndTick, int32 QuantizationInterval, TArray<FGridInterval>& OutSubdivisions ) const
+{
+	FMusicalTime StartMusicalTime = Quantize(GetMusicalTime(StartTick), QuantizationInterval, EQuantizeDirection::Ceiling);
+	FMusicalTime EndMusicalTime = Quantize(GetMusicalTime(EndTick), QuantizationInterval, EQuantizeDirection::Floor);
+
+	if (StartMusicalTime >= EndMusicalTime)
+	{
+		return;
+	}
+
+	while (StartMusicalTime <= EndMusicalTime)
+	{
+		float SubdivisionTime = GetSeconds(StartMusicalTime);
+
+		const FFrameBasedTimeSignaturePoint* TimeSignatureAtFrame = GetTimeSignaturePointForBar(StartMusicalTime.Bar);
+		int32 Tick = TimeSignatureAtFrame->MusicalTimeToTick(StartMusicalTime);
+
+		FGridInterval GridInterval(SubdivisionTime, Tick, StartMusicalTime);
+
+		OutSubdivisions.Add(GridInterval);
+		StartMusicalTime = Add(StartMusicalTime, QuantizationInterval);
+	}
+}
+
+void UFrameBasedMusicMap::GatherBarsForTimespan(int32 StartBar, int32 EndBar, int32 NumBarsInterval, TArray<FGridInterval>& OutSubdivisions) const
+{
+	if (StartBar >= EndBar)
+	{
+		return;
+	}
+
+	while (StartBar <= EndBar)
+	{
+		float BarSeconds = GetSeconds(StartBar);
+
+		const FFrameBasedTimeSignaturePoint* TimeSignatureAtFrame = GetTimeSignaturePointForBar(StartBar);
+		int32 Tick = TimeSignatureAtFrame->BarToTick(StartBar);
+		FMusicalTime MusicalTime = TimeSignatureAtFrame->TickToMusicalTime(Tick);
+
+		FGridInterval GridInterval(BarSeconds, Tick, MusicalTime);
+		OutSubdivisions.Add(GridInterval);
+
+		StartBar = StartBar + NumBarsInterval;
+	}
+}
+
+FMusicalTime UFrameBasedMusicMap::GetMusicalTime(int32 Tick) const
+{
+	if (IsEmpty())
+	{
+		return FMusicalTime();
+	}
+
+	const FFrameBasedTimeSignaturePoint* TimeSignatureAtFrame = GetTimeSignaturePointForTick(Tick);
+	const FFrameBasedTempoPoint* TempoAtFrame = GetTempoPointForTick(Tick);
+
+	return TimeSignatureAtFrame->TickToMusicalTime(Tick);
+}
+
+FMusicalTime UFrameBasedMusicMap::GetMusicalTime(const FFrameTime& FrameTime) const
+{
+	if (IsEmpty())
+	{
+		return FMusicalTime();
+	}
+
+	const FFrameBasedTimeSignaturePoint* TimeSignatureAtFrame = GetTimeSignaturePointForFrameTime(FrameTime);
+	const FFrameBasedTempoPoint* TempoAtFrame = GetTempoPointForFrameTime(FrameTime);
+
+	double Seconds = FrameResolution.AsSeconds(FrameTime);
+	float Tick = TempoAtFrame->SecondsToTick(Seconds);
+	return TimeSignatureAtFrame->TickToMusicalTime(Tick);
+}
+
+FString UFrameBasedMusicMap::MakeBarBeatString(double DecimalFrame, const FFrameBasedTimeSignaturePoint* TimeSignature, const FFrameBasedTempoPoint* Tempo, FrameBasedMusicMap::EBarBeatFormat BarBeatFormat) const
+{
+	double Seconds = FrameResolution.AsSeconds(FFrameTime::FromDecimal(DecimalFrame));
+	float Tick = Tempo->SecondsToTick(Seconds);
+	FMusicalTime MusicalTime = TimeSignature->TickToMusicalTime(Tick);
+	// for display to users... bars and beats are 1 based!
+	MusicalTime.Bar += (BarBeatFormat == FrameBasedMusicMap::EBarBeatFormat::Position ? 1 : 0);
+	double BeatAsDouble = FMath::RoundToDouble((static_cast<float>(MusicalTime.FractionalBeatInBar()) + (BarBeatFormat == FrameBasedMusicMap::EBarBeatFormat::Position ? 1.0f : 0.0f)) * 100.0f);
+	if (BeatAsDouble >= ((TimeSignature->TimeSignature.Numerator + (BarBeatFormat == FrameBasedMusicMap::EBarBeatFormat::Position ? 1 : 0)) * 100))
+	{
+		MusicalTime.Bar++;
+		BeatAsDouble = BarBeatFormat == FrameBasedMusicMap::EBarBeatFormat::Position ? 100.0 : 0.0;
+	}
+	return FString::Printf(TEXT("%d.%.2f"), MusicalTime.Bar, BeatAsDouble / 100.0);
+}
+
+void UFrameBasedMusicMap::GetGridMetricsMusic(float InViewStartSeconds, float InViewEndSeconds, float PixelsPerSecond, float MinTickPx, float DesiredMajorTickPx, TArray<FGridInterval>& OutMajorIntervals, TArray<FGridInterval>& OutMinorIntervals) const
+{
+	const FFrameBasedTempoPoint* StartTempoPoint = GetTempoPointForSeconds(InViewStartSeconds);
+	int32 StartTick = StartTempoPoint->SecondsToTick(InViewStartSeconds);
+	int32 QuantizedStartTick = StartTick;
+
+	FMusicalTime StartMusicalTime = GetMusicalTime(StartTick);
+	int32 StartBar = StartMusicalTime.Bar;
+
+	int32 SubdivisionStepTickSize = FrameBasedMusicMap::EWholeQuantizationTicks::ThirtysecondNote;
+	int32 BarStepSize = 1;
+
+	const float DesiredMajorTickSeconds = DesiredMajorTickPx / PixelsPerSecond;
+
+	//We find a subdivision size that intersects with our ViewStart and fits our desired size
+	bool bFoundFirstGridLine = false;
+	while (!bFoundFirstGridLine)
+	{
+		if (SubdivisionStepTickSize <= FrameBasedMusicMap::EWholeQuantizationTicks::HalfNote)
+		{
+			QuantizedStartTick = FMath::FloorToDouble((float)StartTick / (float)SubdivisionStepTickSize) * SubdivisionStepTickSize;
+			
+			StartTempoPoint = GetTempoPointForTick(QuantizedStartTick);
+			double QuantizedStartSeconds = StartTempoPoint->TickToSeconds(QuantizedStartTick);
+
+			int32 SteppedTick = QuantizedStartTick + SubdivisionStepTickSize;
+			const FFrameBasedTempoPoint* StartSteppedTempoPoint = GetTempoPointForTick(SteppedTick);
+			double SteppedSeconds = StartSteppedTempoPoint->TickToSeconds(SteppedTick);
+
+			float StepSizeSeconds = SteppedSeconds - QuantizedStartSeconds;
+
+			//Make sure were not stepping through a downbeat of next bar
+			StartMusicalTime = GetMusicalTime(QuantizedStartTick);
+			FMusicalTime SteppedMusicalTime = GetMusicalTime(SteppedTick);
+
+			bool bValidStep = false;
+			if (StartMusicalTime.Bar == SteppedMusicalTime.Bar || SteppedMusicalTime.TickInBar == 0)
+			{
+				bValidStep = true;
+			}
+
+			if (StepSizeSeconds > DesiredMajorTickSeconds && bValidStep)
+			{
+				bFoundFirstGridLine = true;
+
+				FGridInterval NonLinearGridInterval = FGridInterval(QuantizedStartSeconds, QuantizedStartTick, GetMusicalTime(QuantizedStartTick));
+				OutMajorIntervals.Add(NonLinearGridInterval);
+			}
+			else
+			{
+				SubdivisionStepTickSize *= 2;
+			}
+		}
+		else
+		{
+			int32 QuantizedStartBar = FMath::FloorToDouble((float)StartBar / (float)BarStepSize) * BarStepSize;
+			float QuantizedBarStartTime = GetSeconds(QuantizedStartBar);
+			const FFrameBasedTempoPoint* StartSteppedTempoPoint = GetTempoPointForSeconds(QuantizedBarStartTime);
+			QuantizedStartTick = StartSteppedTempoPoint->SecondsToTick(QuantizedBarStartTime);
+
+			int32 SteppedBar = QuantizedStartBar + BarStepSize;
+
+			float SecondsAtSteppedBar = GetSeconds(SteppedBar);
+
+			float StepSizeSeconds = SecondsAtSteppedBar - QuantizedBarStartTime;
+
+			if (StepSizeSeconds > DesiredMajorTickSeconds || BarStepSize >= 2048)
+			{
+				bFoundFirstGridLine = true;
+
+				FGridInterval NonLinearGridInterval = FGridInterval(QuantizedBarStartTime, QuantizedStartTick, BarsToMusicalTime(QuantizedStartBar));
+				OutMajorIntervals.Add(NonLinearGridInterval);
+			}
+			else
+			{
+				BarStepSize *= 2;
+			}
+		}
+	}
+
+	const FFrameBasedTempoPoint* EndTempoPoint = GetTempoPointForSeconds(InViewEndSeconds);
+	int32 EndTick = EndTempoPoint->SecondsToTick(InViewEndSeconds);
+
+	int32 StepSizeTicks = 0;
+
+	for (int32 CurrentTick = QuantizedStartTick; CurrentTick < EndTick; CurrentTick += StepSizeTicks)
+	{
+		ComputeGridSpacingAtTick(
+			CurrentTick,
+			DesiredMajorTickSeconds,
+			StepSizeTicks);
+
+		int32 SteppedTicks = CurrentTick + StepSizeTicks;
+
+		FMusicalTime CurrentMusicalTime = GetMusicalTime(CurrentTick);
+		FMusicalTime SteppedMusicalTime = GetMusicalTime(SteppedTicks);
+
+		//After computing spacing, we may need to correct past intervals if the spacing dosen't make sense
+		if (SteppedMusicalTime.Bar != CurrentMusicalTime.Bar && CurrentMusicalTime.Bar > 0)
+		{
+			int32 NumBarsStepped = SteppedMusicalTime.Bar - CurrentMusicalTime.Bar;
+
+			//Make sure our step size dosen't take us into next bar, i.e. always favour marker on downbeat
+			if (SteppedMusicalTime.TickInBar > 0)
+			{
+				OutMajorIntervals.Pop();
+				
+				CurrentTick = OutMajorIntervals.Last().IntervalTick;
+				CurrentMusicalTime = OutMajorIntervals.Last().MusicalTime;
+
+				StepSizeTicks = CurrentMusicalTime.TicksPerBar - CurrentMusicalTime.TickInBar;
+				SteppedTicks = CurrentTick + StepSizeTicks;
+				SteppedMusicalTime = GetMusicalTime(SteppedTicks);
+			}
+			//Make sure we aren't stepping up in multiple number of bars at wrong interval
+			else if (SteppedMusicalTime.Bar % NumBarsStepped != 0)
+			{
+				//We need to increase the previous interval 
+				FMusicalTime PreviousSteppedMusicalTime = OutMajorIntervals.Last().MusicalTime;
+
+				OutMajorIntervals.Pop();
+
+				if (OutMajorIntervals.Num() > 0)
+				{
+					FMusicalTime PreviousMusicalTime = OutMajorIntervals.Last().MusicalTime;
+
+					NumBarsStepped = PreviousSteppedMusicalTime.Bar - PreviousMusicalTime.Bar;
+					NumBarsStepped *= 2;
+
+					int32 SteppedBars = PreviousMusicalTime.Bar + NumBarsStepped;
+
+					const FFrameBasedTimeSignaturePoint* TimeSigPoint = GetTimeSignaturePointForBar(SteppedBars);
+					int32 BarsAtTimeSig = SteppedBars - TimeSigPoint->OnBar;
+
+					StepSizeTicks = (TimeSigPoint->OnTick - CurrentTick) + BarsAtTimeSig * TimeSigPoint->TicksPerBar();
+					SteppedTicks = CurrentTick + StepSizeTicks;
+					SteppedMusicalTime = GetMusicalTime(SteppedTicks);
+				}
+			}
+		}
+
+		const FFrameBasedTempoPoint* SteppedTempoPoint = GetTempoPointForTick(SteppedTicks);
+		float SteppedSizeSeconds = SteppedTempoPoint->TickToMs(SteppedTicks) / 1000.f;
+
+		FGridInterval NonLinearGridInterval = FGridInterval(SteppedSizeSeconds, SteppedTicks, SteppedMusicalTime);
+		OutMajorIntervals.Add(NonLinearGridInterval);
+	}
+
+	for (int32 Idx = 0; Idx < OutMajorIntervals.Num() - 1; ++Idx)
+	{
+		ComputeMinorGridSpacing(OutMajorIntervals[Idx].IntervalTick, OutMajorIntervals[Idx + 1].IntervalTick, OutMinorIntervals);
+	}
+}
+
+void UFrameBasedMusicMap::ComputeGridSpacingAtTick(int32 TickToCompute, float DesiredMajorTickSeconds, int32& OutMajorIntervalTicks) const
+{
+	const FFrameBasedTempoPoint* TempoPoint = GetTempoPointForTick(TickToCompute);
+	float TimeToComputeSeconds = TempoPoint->TickToSeconds(TickToCompute);
+
+	float DesiredTimeStepped = TimeToComputeSeconds + DesiredMajorTickSeconds;
+
+	int32 TickStepped = TempoPoint->MsToTick(DesiredTimeStepped * 1000.f);
+
+	int32 StepTickSize = TickStepped - TickToCompute;
+
+	//This will round up to lowest 2^N tick size that is larger than our desired size
+	int32 PowTwo = FMath::RoundUpToPowerOfTwo(StepTickSize / FrameBasedMusicMap::EWholeQuantizationTicks::ThirtysecondNote);
+	StepTickSize = FrameBasedMusicMap::EWholeQuantizationTicks::ThirtysecondNote * PowTwo;
+
+	FMusicalTime StartMusicalTime = GetMusicalTime(TickToCompute);
+	if (StartMusicalTime.TickInBar == 0 && StepTickSize > FrameBasedMusicMap::EWholeQuantizationTicks::HalfNote)
+	{
+		FMusicalTime SteppedMusicalTime = GetMusicalTime(TickStepped);
+
+		int32 NumBarsStepped = SteppedMusicalTime.Bar - StartMusicalTime.Bar;
+		NumBarsStepped++;
+
+		NumBarsStepped = FMath::Clamp(NumBarsStepped, 1, 2048);
+		int32 NumMultipleBarsStepped = FMath::RoundUpToPowerOfTwo(NumBarsStepped);
+
+		int32 SteppedBars = StartMusicalTime.Bar + NumMultipleBarsStepped;
+
+		const FFrameBasedTimeSignaturePoint* TimeSigPoint = GetTimeSignaturePointForBar(SteppedBars);
+		int32 BarsAtTimeSig = SteppedBars - TimeSigPoint->OnBar;
+
+		StepTickSize = (TimeSigPoint->OnTick - TickToCompute) + BarsAtTimeSig * TimeSigPoint->TicksPerBar();
+	}
+	else
+	{
+		//We can't step upwards if our starting tick isn't on a subdivision boundary of that size
+		//I.e. if we are on the tick 3/8th notes into a bar, we can't start stepping in 1/4 notes
+		while (StartMusicalTime.TickInBar % StepTickSize != 0 && StepTickSize > FrameBasedMusicMap::EWholeQuantizationTicks::ThirtysecondNote)
+		{
+			StepTickSize /= 2;
+		}
+
+		if (StepTickSize < FrameBasedMusicMap::EWholeQuantizationTicks::ThirtysecondNote)
+		{
+			StepTickSize = FrameBasedMusicMap::EWholeQuantizationTicks::ThirtysecondNote;
+		}
+	}
+
+	OutMajorIntervalTicks = StepTickSize;
+}
+
+void UFrameBasedMusicMap::ComputeMinorGridSpacing(int32 StartTick, int32 EndTick, TArray<FGridInterval>& OutMinorIntervals) const
+{
+	int32 StepTickSize = EndTick - StartTick;
+    if (StepTickSize <= 0)
+    {
+    	return;
+    }
+	FMusicalTime StartMusicalTime = GetMusicalTime(StartTick);
+	FMusicalTime EndMusicalTime = GetMusicalTime(EndTick);
+
+	int32 BarStepSize = EndMusicalTime.Bar - StartMusicalTime.Bar;
+
+	if (StepTickSize >= StartMusicalTime.TicksPerBar && BarStepSize >= 4)
+	{
+		int32 MinorIntervalBars = BarStepSize / 4;
+		GatherBarsForTimespan(StartMusicalTime.Bar, EndMusicalTime.Bar, MinorIntervalBars, OutMinorIntervals);
+	}
+	else
+	{
+		int32 MinorTickSize = StartMusicalTime.TicksPerBeat;
+		while (MinorTickSize >= StepTickSize)
+		{
+			MinorTickSize /= 2;
+		}
+		MinorTickSize = FMath::Clamp(MinorTickSize, FrameBasedMusicMap::EWholeQuantizationTicks::ThirtysecondNote, FrameBasedMusicMap::EWholeQuantizationTicks::HalfNote);
+		GatherSubdivisonsForTimespan(StartTick, EndTick, MinorTickSize, OutMinorIntervals);
+	}
+}
+
+TArray<FFrameBasedTempoPoint> UFrameBasedMusicMap::GetTempoMetrics(float InViewStartSeconds, float InViewEndSeconds) const
+{
+	TArray<FFrameBasedTempoPoint> TempoPoints;
+
+	float ViewStartMs = InViewStartSeconds * 1000.f;
+
+	float ViewEndMs = InViewEndSeconds * 1000.f;
+
+	for (FFrameBasedTempoPoint TempoPoint : TempoMap)
+	{
+		if (TempoPoint.OnMs > ViewStartMs && TempoPoint.OnMs < ViewEndMs)
+		{
+			TempoPoints.Add(TempoPoint);
+		}
+	}
+
+	return TempoPoints;
+}
+
+void UFrameBasedMusicMap::GetTempoTimeSignatureMarkers(float InStartSeconds, float InEndSeconds, TArray<FTempoTimeSignatureMarker>& OutMarkers) const
+{
+	float ViewStartMs = InStartSeconds * 1000.f;
+	float ViewEndMs = InEndSeconds * 1000.f;
+
+	const FFrameBasedTempoPoint* StartTempoPoint = GetTempoPointForSeconds(InStartSeconds);
+	int32 ViewStartTick = StartTempoPoint->SecondsToTick(InStartSeconds);
+
+	const FFrameBasedTempoPoint* EndTempoPoint = GetTempoPointForSeconds(InEndSeconds);
+	int32 ViewEndTick = EndTempoPoint->SecondsToTick(InEndSeconds);
+
+	for (const FFrameBasedTempoPoint& TempoPoint : TempoMap)
+	{
+		if (TempoPoint.OnMs > ViewStartMs && TempoPoint.OnMs < ViewEndMs)
+		{
+			FTempoTimeSignatureMarker Marker;
+			Marker.Seconds = TempoPoint.OnMs / 1000.f;
+			Marker.TempoPoint = TempoPoint;
+
+			OutMarkers.Add(Marker);
+		}
+	}
+
+	for (const FFrameBasedTimeSignaturePoint& TimeSigPoint : BarMap)
+	{
+		if (TimeSigPoint.OnTick > ViewStartTick && TimeSigPoint.OnTick < ViewEndTick)
+		{
+			if (FTempoTimeSignatureMarker* FoundMarker = OutMarkers.FindByPredicate([TimeSigPoint](const FTempoTimeSignatureMarker& Marker) { return Marker.TempoPoint.IsSet() && Marker.TempoPoint->OnTick == TimeSigPoint.OnTick; }))
+			{
+				FoundMarker->TimeSigPoint = TimeSigPoint;
+			}
+			else
+			{
+				FTempoTimeSignatureMarker Marker;
+				Marker.TimeSigPoint = TimeSigPoint;
+
+				const FFrameBasedTempoPoint* TempoPoint = GetTempoPointForTick(TimeSigPoint.OnTick);
+				Marker.Seconds = TempoPoint->TickToSeconds(TimeSigPoint.OnTick);
+
+				OutMarkers.Add(Marker);
+			}
+		}
+	}
+}
+
+
+#undef LOCTEXT_NAMESPACE
diff --git a/Engine/Plugins/Runtime/MusicEnvironment/Source/MusicEnvironment/Private/MusicClockSourceManager.cpp b/Engine/Plugins/Runtime/MusicEnvironment/Source/MusicEnvironment/Private/MusicClockSourceManager.cpp
new file mode 100644
--- /dev/null
+++ b/Engine/Plugins/Runtime/MusicEnvironment/Source/MusicEnvironment/Private/MusicClockSourceManager.cpp
@@ -0,0 +1,155 @@
+// Copyright Epic Games, Inc. All Rights Reserved.
+
+#include "MusicClockSourceManager.h"
+
+#include "MusicEnvironmentSubsystem.h" 
+
+TScriptInterface<IMusicEnvironmentClockSource> UMusicClockSourceManager::FindClock(const FGameplayTag& Tag, bool bExactMatch)
+{
+	CleanMap();
+	if (const TWeakInterfacePtr<IMusicEnvironmentClockSource>* Clock = TaggedClocks.Find(Tag))
+	{
+		return Clock->ToScriptInterface();;
+	}
+	if (!bExactMatch)
+	{
+		int32 NumberOfMatches = 0;
+		TWeakInterfacePtr<IMusicEnvironmentClockSource> MatchingClock;
+		FGameplayTag MatchingTag;
+		for (auto& It : TaggedClocks)
+		{
+			if (It.Key.MatchesTag(Tag))
+			{
+				++NumberOfMatches;
+				if (!MatchingClock.IsValid())
+				{
+					MatchingClock = It.Value;
+					MatchingTag = It.Key;
+				}
+			}
+		}
+		if (MatchingClock.IsValid())
+		{
+			if (NumberOfMatches > 1)
+			{ 
+				UE_LOG(LogMusicEnvironment, Warning, TEXT("FindClock found %d clocks that match the tag '%s'. Returning clock tagged '%s'."), NumberOfMatches, *Tag.ToString(), *MatchingTag.ToString());
+			}
+			return MatchingClock.ToScriptInterface();
+		}
+	}
+	return nullptr;
+}
+
+void UMusicClockSourceManager::AddTaggedClock(const FGameplayTag& Tag, TScriptInterface<IMusicEnvironmentClockSource> InClock)
+{
+	UObject* AsObject = InClock.GetObject();
+	AddTaggedClock(Tag, AsObject);
+}
+
+void UMusicClockSourceManager::AddTaggedClock(const FGameplayTag& Tag, UObject* InClock)
+{
+	CleanMap();
+	TWeakInterfacePtr<IMusicEnvironmentClockSource> WeakInterfacePtr(InClock);
+	if (!WeakInterfacePtr.IsValid())
+	{
+		return;
+	}
+	if (TaggedClocks.Contains(Tag))
+	{
+		if (WeakInterfacePtr != TaggedClocks[Tag])
+		{
+			UE_LOG(LogMusicEnvironment, Warning, TEXT("The MusicClockSourceManager is already tracking a clock tagged '%s'. It will be replaced by this new request."), *Tag.ToString());
+			TaggedClocks[Tag] = WeakInterfacePtr;
+		}
+	}
+	else
+	{
+		TaggedClocks.Add(Tag, WeakInterfacePtr);
+	}
+}
+
+void UMusicClockSourceManager::RemoveTaggedClock(TScriptInterface<IMusicEnvironmentClockSource> InClock)
+{
+	UObject* AsObject = InClock.GetObject();
+	RemoveTaggedClock(AsObject);
+}
+
+void UMusicClockSourceManager::RemoveTaggedClock(UObject* InClock)
+{
+	CleanMap();
+	TWeakInterfacePtr<IMusicEnvironmentClockSource> WeakInterfacePtr(InClock);
+	if (WeakInterfacePtr.IsValid())
+	{
+		TaggedClocks = TaggedClocks.FilterByPredicate([&](const TPair<FGameplayTag, TWeakInterfacePtr<IMusicEnvironmentClockSource>>& ClockEntry) { return ClockEntry.Value != WeakInterfacePtr; });
+	}
+}
+
+void UMusicClockSourceManager::RemoveClockWithTag(const FGameplayTag& Tag)
+{
+	CleanMap();
+	TaggedClocks.Remove(Tag);
+}
+
+TScriptInterface<IMusicEnvironmentClockSource> UMusicClockSourceManager::GetGlobalMusicClockAuthority()
+{
+	CleanStack();
+	if (GlobalClockSourceStack.IsEmpty())
+	{
+		return nullptr;
+	}
+	return GlobalClockSourceStack.Last().ToScriptInterface();
+}
+
+bool UMusicClockSourceManager::PushGlobalMusicClockAuthority(TScriptInterface<IMusicEnvironmentClockSource> InClock)
+{
+	UObject* AsObject = InClock.GetObject();
+	return PushGlobalMusicClockAuthority(AsObject);
+}
+
+bool UMusicClockSourceManager::PushGlobalMusicClockAuthority(UObject* InClock)
+{
+	CleanStack();
+	TWeakInterfacePtr<IMusicEnvironmentClockSource> WeakInterfacePtr(InClock);
+	if (!WeakInterfacePtr.IsValid())
+	{
+		return false;
+	}
+	GlobalClockSourceStack.Add(WeakInterfacePtr);
+	return true;
+}
+
+void UMusicClockSourceManager::RemoveGlobalClockAuthority(TScriptInterface<IMusicEnvironmentClockSource> InClock)
+{
+	UObject* AsObject = InClock.GetObject();
+	RemoveGlobalClockAuthority(AsObject);
+}
+
+void UMusicClockSourceManager::RemoveGlobalClockAuthority(UObject* InClock)
+{
+	CleanStack();
+	TWeakInterfacePtr<IMusicEnvironmentClockSource> WeakInterfacePtr(InClock);
+	if (!WeakInterfacePtr.IsValid())
+	{
+		return;
+	}
+	GlobalClockSourceStack.RemoveAll([&](const TWeakInterfacePtr<IMusicEnvironmentClockSource>& ClockEntry) { return ClockEntry == WeakInterfacePtr; });
+}
+
+void UMusicClockSourceManager::PopMusicClockAuthority()
+{
+	CleanStack();
+	if (!GlobalClockSourceStack.IsEmpty())
+	{
+		GlobalClockSourceStack.Pop(EAllowShrinking::No);
+	}
+}
+
+void UMusicClockSourceManager::CleanMap()
+{
+	TaggedClocks = TaggedClocks.FilterByPredicate([](const TPair<FGameplayTag, TWeakInterfacePtr<IMusicEnvironmentClockSource>>& ClockEntry){ return ClockEntry.Value.IsValid(); });
+}
+
+void UMusicClockSourceManager::CleanStack()
+{
+	GlobalClockSourceStack.RemoveAll([](const TWeakInterfacePtr<IMusicEnvironmentClockSource>& ClockEntry){ return !ClockEntry.IsValid(); });
+}
diff --git a/Engine/Plugins/Runtime/MusicEnvironment/Source/MusicEnvironment/Private/MusicEnvironmentMetronome.cpp b/Engine/Plugins/Runtime/MusicEnvironment/Source/MusicEnvironment/Private/MusicEnvironmentMetronome.cpp
new file mode 100644
--- /dev/null
+++ b/Engine/Plugins/Runtime/MusicEnvironment/Source/MusicEnvironment/Private/MusicEnvironmentMetronome.cpp
@@ -0,0 +1,19 @@
+// Copyright Epic Games, Inc. All Rights Reserved.
+
+#include "MusicEnvironmentMetronome.h"
+
+// Add default functionality here for any IMovieSceneMetronome functions that are not pure virtual.
+void IMusicEnvironmentMetronome::SetMusicMap(UFrameBasedMusicMap* InMusicMap)
+{
+	MusicMap = TStrongObjectPtr<UFrameBasedMusicMap>(InMusicMap);
+	OnMusicMapSet();
+}
+
+bool IMusicEnvironmentMetronome::SetTempo(const float Bpm)
+{
+	if (MusicMap.IsValid() && !MusicMap->IsSimple())
+	{
+		return false;
+	}
+	return OnSetTempo(Bpm);
+}
diff --git a/Engine/Plugins/Runtime/MusicEnvironment/Source/MusicEnvironment/Private/MusicEnvironmentModule.cpp b/Engine/Plugins/Runtime/MusicEnvironment/Source/MusicEnvironment/Private/MusicEnvironmentModule.cpp
new file mode 100644
--- /dev/null
+++ b/Engine/Plugins/Runtime/MusicEnvironment/Source/MusicEnvironment/Private/MusicEnvironmentModule.cpp
@@ -0,0 +1,7 @@
+// Copyright Epic Games, Inc. All Rights Reserved.
+
+#include "MusicEnvironmentModule.h"
+
+#include "Logging/MessageLog.h"
+
+IMPLEMENT_MODULE(FMusicEnvironmentModule, MusicEnvironment);
diff --git a/Engine/Plugins/Runtime/MusicEnvironment/Source/MusicEnvironment/Private/MusicEnvironmentSubsystem.cpp b/Engine/Plugins/Runtime/MusicEnvironment/Source/MusicEnvironment/Private/MusicEnvironmentSubsystem.cpp
new file mode 100644
--- /dev/null
+++ b/Engine/Plugins/Runtime/MusicEnvironment/Source/MusicEnvironment/Private/MusicEnvironmentSubsystem.cpp
@@ -0,0 +1,48 @@
+// Copyright Epic Games, Inc. All Rights Reserved.
+
+#include "MusicEnvironmentSubsystem.h"
+
+#include "MusicClockSourceManager.h"
+#include "MusicEnvironmentMetronome.h"
+
+#include "Engine/Engine.h"
+
+DEFINE_LOG_CATEGORY(LogMusicEnvironment);
+
+UMusicEnvironmentSubsystem& UMusicEnvironmentSubsystem::Get()
+{
+	return *GEngine->GetEngineSubsystem< UMusicEnvironmentSubsystem>();
+}
+
+UMusicClockSourceManager* UMusicEnvironmentSubsystem::GetClockSourceManager()
+{
+	if (!ClockSourceManager)
+	{
+		ClockSourceManager = NewObject<UMusicClockSourceManager>(this);
+	}
+	return ClockSourceManager;
+}
+
+bool UMusicEnvironmentSubsystem::SetMetronomeClass(const UClass* InMetronomeType)
+{
+	if (!InMetronomeType->ImplementsInterface(UMusicEnvironmentMetronome::StaticClass()))
+	{
+		return false;
+	}
+	MetronomeType = InMetronomeType;
+	return true;
+}
+
+bool UMusicEnvironmentSubsystem::CanSpawnMetronome() const
+{
+	return MetronomeType.IsValid();
+}
+
+TScriptInterface<IMusicEnvironmentMetronome> UMusicEnvironmentSubsystem::SpawnMetronome(UObject* Outer, FName Name)
+{
+	if (!MetronomeType.IsValid())
+	{
+		return nullptr;
+	}
+	return TScriptInterface<IMusicEnvironmentMetronome>(NewObject<UObject>(Outer, MetronomeType.Get(), MakeUniqueObjectName(Outer, MetronomeType.Get(), Name), RF_Transient));
+}
diff --git a/Engine/Plugins/Runtime/MusicEnvironment/Source/MusicEnvironment/Public/FrameBasedMusicMap.h b/Engine/Plugins/Runtime/MusicEnvironment/Source/MusicEnvironment/Public/FrameBasedMusicMap.h
new file mode 100644
--- /dev/null
+++ b/Engine/Plugins/Runtime/MusicEnvironment/Source/MusicEnvironment/Public/FrameBasedMusicMap.h
@@ -0,0 +1,464 @@
+// Copyright Epic Games, Inc. All Rights Reserved.
+
+#pragma once
+
+#include "CoreTypes.h"
+#include "Misc/FrameRate.h"
+#include "Misc/MusicalTime.h"
+#include "UObject/Interface.h"
+#include "Sound/QuartzQuantizationUtilities.h"
+
+#include "FrameBasedMusicMap.generated.h"
+
+struct FFrameTime;
+
+namespace FrameBasedMusicMap
+{
+	enum ETripletQuantizationTicks : int32
+	{
+		ThirtysecondNoteTriplet = (MusicalTime::TicksPerQuarterNote / 4) / 3,	// 80
+		SixteenthNoteTriplet    = (MusicalTime::TicksPerQuarterNote / 2) / 3,	// 160
+		EighthNoteTriplet       = (MusicalTime::TicksPerQuarterNote) / 3,		// 320
+		QuarterNoteTriplet      = (MusicalTime::TicksPerQuarterNote * 2) / 3,	// 640
+		HalfNoteTriplet         = (MusicalTime::TicksPerQuarterNote * 4) / 3,	// 1280
+		WholeNoteTriplet        = (MusicalTime::TicksPerQuarterNote * 8) / 3,	// 2560
+	};
+
+	enum EWholeQuantizationTicks : int32
+	{
+		ThirtysecondNote        = MusicalTime::TicksPerQuarterNote / 8,			// 120
+		SixteenthNote           = MusicalTime::TicksPerQuarterNote / 4,			// 240
+		EighthNote              = MusicalTime::TicksPerQuarterNote / 2,			// 480
+		QuarterNote	            = MusicalTime::TicksPerQuarterNote,				// 960
+		HalfNote                = MusicalTime::TicksPerQuarterNote * 2,			// 1920
+		WholeNote               = MusicalTime::TicksPerQuarterNote * 4,			// 3840
+		Bar                     = std::numeric_limits<int32>::min()+1,          // Ticks depends on Time Signature
+		Beat                    = std::numeric_limits<int32>::min()+2           // Ticks depends on Time Signature
+	};
+
+	enum EDottedQuantizationTicks : int32
+	{
+		DottedThirtysecondNote  = (MusicalTime::TicksPerQuarterNote / 16) * 3,	// 180
+		DottedSixteenthNote     = (MusicalTime::TicksPerQuarterNote / 8) * 3,	// 360
+		DottedEighthNote        = (MusicalTime::TicksPerQuarterNote / 4) * 3,	// 720
+		DottedQuarterNote       = (MusicalTime::TicksPerQuarterNote / 2) * 3,	// 1440
+		DottedHalfNote          = (MusicalTime::TicksPerQuarterNote) * 3,		// 2880
+		DottedWholeNote         = (MusicalTime::TicksPerQuarterNote) * 6,		// 5760
+	};
+
+	MUSICENVIRONMENT_API FText GetQuantizationName(int32 Quantization);
+	MUSICENVIRONMENT_API EQuartzCommandQuantization QuantiazationIntervalToQuartz(int32 Quantization);
+
+	extern MUSICENVIRONMENT_API TArray<TSharedPtr<int32>> QuantizationOptionsSource;
+
+	enum class EBarBeatFormat
+	{
+		Position,
+		Duration,
+		NotAllowed // Don't display bars.beats! Fall back to seconds.
+	};
+}
+
+struct MUSICENVIRONMENT_API FGridInterval
+{
+	double IntervalStartTime;
+
+	int32 IntervalTick;
+
+	FMusicalTime MusicalTime;
+
+	FGridInterval(double InIntervalStartTime, int32 InIntervalTick, FMusicalTime InMusicalTime)
+	{
+		IntervalStartTime = InIntervalStartTime;
+		IntervalTick = InIntervalTick;
+		MusicalTime = InMusicalTime;
+	}
+};
+
+USTRUCT()
+struct MUSICENVIRONMENT_API FFrameBasedTempoPoint
+{
+	GENERATED_BODY()
+
+	/**
+	 * NOTE: This structure treats BPM/Tempo exactly like Midi does to avoid any
+	 * floating point precision problems and differences between Midi code and 
+	 * this code. It is for that reason that Tempo is stored here a "microseconds
+	 * per quarter note" (MPQ) and not a float Bpm. 
+	 */
+
+	FFrameBasedTempoPoint()
+	{}
+
+	FFrameBasedTempoPoint(float InBpm, int32 InTick, float InAtMillisecond)
+		: OnTick(InTick)
+		, OnMs(InAtMillisecond)
+	{
+		SetBpm(InBpm);
+	}
+
+	void SetBpm(float InBpm)
+	{
+		MicrosecondsPerQuarterNote = InBpm >= 0.0f ? static_cast<int32>(60000000.0f / InBpm) : 0;
+	}
+
+	UPROPERTY()
+	int32 MicrosecondsPerQuarterNote = 500000; // 120 bpm default.
+	UPROPERTY()
+	int32 OnTick = 0;
+	UPROPERTY()
+	float OnMs = 0.0;
+
+	float Bpm() const
+	{
+		return MicrosecondsPerQuarterNote == 0 ? 0 : 60000000.0f / static_cast<float>(MicrosecondsPerQuarterNote);
+	}
+
+	int32 MsToTick(float Ms) const
+	{
+		if (Ms == 0.0f)
+		{
+			return 0;
+		}
+
+		return FMath::RoundToInt32(static_cast<double>(OnTick) +
+								static_cast<double>(Ms - OnMs) * 1000.0 / 
+								static_cast<double>(MicrosecondsPerQuarterNote) *
+								static_cast<double>(MusicalTime::TicksPerQuarterNote));
+	}
+
+	int32 SecondsToTick(double Second) const
+	{
+		return MsToTick(static_cast<float>(Second * 1000.0));
+	}
+
+	double TickToSeconds(int32 Tick) const
+	{
+		return static_cast<double>(TickToMs(Tick))/1000.0;
+	}
+
+	float TickToMs(int32 Tick) const
+	{
+		if (Tick == 0)
+		{
+			return 0.0f;
+		}
+		float MsPerTick = static_cast<float>(MicrosecondsPerQuarterNote) / (static_cast<float>(MusicalTime::TicksPerQuarterNote) * 1000.0f);
+		float DeltaTick = static_cast<float>(Tick - OnTick);
+		return OnMs + (DeltaTick * MsPerTick);
+	}
+};
+
+USTRUCT()
+struct MUSICENVIRONMENT_API FFrameBasedTimeSignature
+{
+	GENERATED_BODY()
+
+	FFrameBasedTimeSignature()
+	{}
+	FFrameBasedTimeSignature(int32 InNumerator, int32 InDenominator)
+		: Numerator(static_cast<int16>(InNumerator))
+		, Denominator(static_cast<int16>(InDenominator))
+	{}
+
+	UPROPERTY()
+	int16 Numerator = 4;
+	UPROPERTY()
+	int16 Denominator = 4;
+
+	bool operator==(const FFrameBasedTimeSignature& Other)
+	{
+		return Numerator == Other.Numerator && Denominator == Other.Denominator;
+	}
+
+	bool operator!=(const FFrameBasedTimeSignature& Other)
+	{
+		return Numerator != Other.Numerator || Denominator != Other.Denominator;
+	}
+};
+
+USTRUCT()
+struct MUSICENVIRONMENT_API FFrameBasedTimeSignaturePoint
+{
+	GENERATED_BODY()
+
+	FFrameBasedTimeSignaturePoint()
+	{}
+
+	FFrameBasedTimeSignaturePoint(const FFrameBasedTimeSignature& InTimeSignature, int32 InOnBar, int32 InOnTick, double InOnFrame)
+		: TimeSignature(InTimeSignature)
+		, OnBar(InOnBar)
+		, OnTick(InOnTick)
+		, OnFrame(InOnFrame)
+	{
+	}
+
+	UPROPERTY()
+	FFrameBasedTimeSignature TimeSignature;
+	UPROPERTY()
+	int32 OnBar = 0;
+	UPROPERTY()
+	int32 OnTick = 0;
+	UPROPERTY()
+	double OnFrame = 0.0;
+
+	int32 TicksPerBeat() const
+	{
+		return MusicalTime::TicksPerQuarterNote * 4 / TimeSignature.Denominator;
+	}
+
+	int32 TicksPerBar() const
+	{
+		return TicksPerBeat() * TimeSignature.Numerator;
+	}
+
+	FMusicalTime TickToMusicalTime(int32 InTick) const
+	{
+		FMusicalTime Result(0, 0, TicksPerBar(), TicksPerBeat());
+		int32 TicksPast = InTick - OnTick;
+		int32 BarsPassed = TicksPast / TicksPerBar();
+		if (TicksPast >= 0)
+		{
+			Result.Bar = OnBar + BarsPassed;
+			Result.TickInBar = TicksPast - (BarsPassed * TicksPerBar());
+		}
+		else
+		{
+			Result.Bar = OnBar + BarsPassed;
+			int32 TicksRemainder = FMath::Abs(TicksPast) % TicksPerBar();
+			if (TicksRemainder > 0)
+			{
+				Result.Bar--;
+				Result.TickInBar = TicksPerBar() - TicksRemainder;
+			}
+		}
+		return Result;
+	}
+
+	int32 MusicalTimeToTick(const FMusicalTime& MusicalTime) const
+	{
+		int32 BarDelta = MusicalTime.Bar - OnBar;
+		return OnTick + (BarDelta * MusicalTime.TicksPerBar) + MusicalTime.TickInBar;
+	}
+
+	int32 BarToTick(int32 Bar) const
+	{
+		int32 BarDelta = Bar - OnBar;
+		return OnTick + (BarDelta * TicksPerBar());
+	}
+
+	int32 BarBeatToTick(int32 Bar, float Beat) const
+	{
+		int32 BarDelta = Bar - OnBar;
+		float TicksIntoBar = Beat * static_cast<float>(TicksPerBeat());
+		return OnTick + (BarDelta * TicksPerBar()) + FMath::RoundToInt32(TicksIntoBar);
+	}
+};
+
+USTRUCT()
+struct MUSICENVIRONMENT_API FTempoTimeSignatureMarker
+{
+	GENERATED_BODY()
+
+	UPROPERTY()
+	TOptional<FFrameBasedTimeSignaturePoint> TimeSigPoint;
+
+	UPROPERTY()
+	TOptional<FFrameBasedTempoPoint> TempoPoint;
+
+	UPROPERTY()
+	float Seconds = 0.f;
+};
+
+UCLASS()
+class MUSICENVIRONMENT_API UFrameBasedMusicMap : public UObject
+{
+	GENERATED_BODY()
+
+public:
+	UFrameBasedMusicMap()
+		: FrameResolution(0,0)
+	{}
+
+	UFrameBasedMusicMap(const FFrameRate& InFrameResolution)
+		: FrameResolution(InFrameResolution)
+	{
+		
+	}
+
+	UFrameBasedMusicMap(const FFrameRate& InFrameResolution, float InBpm, const FFrameBasedTimeSignature& TimeSignature)
+		: FrameResolution(InFrameResolution)
+	{
+		TempoMap.Add(FFrameBasedTempoPoint(InBpm, 0, 0.0f));
+		BarMap.Add(FFrameBasedTimeSignaturePoint(TimeSignature, 0, 0, 0.0));
+	}
+
+	void SetFrameResolution(const FFrameRate& NewFrameResolution);
+
+	void Init(float InBpm, const FFrameBasedTimeSignature& InTimeSignature);
+
+	void SetInitialTempo(float Bpm);
+	float GetInitialTempo() const;
+	void SetInitialTimeSignature(const FFrameBasedTimeSignature& TimeSignature);
+	FFrameBasedTimeSignature GetInitialTimeSignature() const;
+
+	void Clear()
+	{
+		TempoMap.Empty();
+		BarMap.Empty();
+	}
+
+	bool IsEmpty() const { return TempoMap.IsEmpty() || BarMap.IsEmpty(); }
+	bool IsSimple() const { return TempoMap.Num() == 1 && BarMap.Num() == 1; }
+
+	double GetFrame(const FMusicalTime& MusicalTime) const;
+	double GetFrame(int32 AtBar) const;
+	FFrameTime GetFrameTime(const FMusicalTime& MusicalTime) const;
+	FFrameTime GetFrameTime(int32 AtBar, float AtBeat) const;
+	FFrameTime GetFrameTime(double FractionalBar) const;
+	FFrameTime GetFrameTime(int32 AtTick) const;
+
+	double GetSeconds(const FMusicalTime& MusicalTime) const;
+	double GetSeconds(int32 AtBar, float AtBeat = 0.0f) const;
+
+	FString MakeBarBeatStringForFrame(double DecimalFrame, FrameBasedMusicMap::EBarBeatFormat BarBeatFormat) const;
+	FString MakeBarBeatStringForFrameTime(const FFrameTime& Frame, FrameBasedMusicMap::EBarBeatFormat BarBeatFormat) const;
+	FString MakeBarBeatStringForMusicalTime(const FMusicalTime& MusicalTime, FrameBasedMusicMap::EBarBeatFormat BarBeatFormat) const;
+	FString MakeBarBeatStringForSeconds(double Seconds, FrameBasedMusicMap::EBarBeatFormat BarBeatFormat) const;
+
+	FMusicalTime GetMusicalTime(int32 Tick) const;
+	FMusicalTime GetMusicalTime(const FFrameTime& FrameTime) const;
+
+	float GetTempoForFrame(double DecimalFrame) const;
+	float GetTempoForFrameTime(const FFrameTime& FrameTime) const;
+	float GetTempoForSeconds(double Seconds) const;
+
+	FFrameBasedTimeSignature GetTimeSignatureForFrame(double DecimalFrame) const;
+	FFrameBasedTimeSignature GetTimeSignatureForFrameTime(const FFrameTime& FrameTime) const;
+	FFrameBasedTimeSignature GetTimeSignatureForSeconds(double Seconds) const;
+	int32 GetBeatsInBar(int32 Bar) const;
+	int32 GetTicksInBar(int32 Bar) const;
+	
+	FFrameTime RemapFrameTime(FFrameTime SourceTime, const UFrameBasedMusicMap* OldMap) const;
+
+	FMusicalTime BarsToMusicalTime(double BarsFromStart) const;
+
+	enum class EQuantizeDirection
+	{
+		Nearest = 0,
+		Ceiling = 1,
+		Floor = 2
+	};
+	FMusicalTime Quantize(const FMusicalTime& MusicalTime, int32 QuantizationInterval, EQuantizeDirection Direction = EQuantizeDirection::Nearest) const;
+	int32 QuantizeToTick(const FMusicalTime& MusicalTime, int32 QuantizationInterval, EQuantizeDirection Direction = EQuantizeDirection::Nearest) const;
+
+	/**
+	* FMusicalTime Add(const FMusicalTime& MusicalTime, int32 QuantizationInterval) const;
+	* 
+	* Increment the supplied MusicalTime by the specified QuantizationInterval. 
+	* NOTE: This does NOT quantize anything. QuantizationInterval is just used 
+	* as a convenient way to specify a musical subdivision. You can use any of the 
+	* TripletQuantizationTicks, WholeQuantizationTicks, or DottedQuantizationTicks
+	* enum values. AND... you can negate them to go back in time.
+	* 
+	* @param MusicalTime			Time to be added to
+	
+	* @param QuantizationInterval	One of the values from...
+	*                               TripletQuantizationTicks
+	*								WholeQuantizationTicks
+	*								DottedQuantizationTicks
+	* 
+	*                                 
+	*/
+	FMusicalTime Add(const FMusicalTime& MusicalTime, int32 QuantizationInterval) const;
+
+	UPROPERTY()
+	TArray<FFrameBasedTempoPoint> TempoMap;
+	UPROPERTY()
+	TArray<FFrameBasedTimeSignaturePoint> BarMap;
+	UPROPERTY()
+	FFrameRate FrameResolution;
+
+	void AddTimeSignature(int32 AtTick, int32 AtBar, int16 Numerator, int16 Denominator);
+	void AddTempo(int32 AtTick, float AtMs, float Bpm);
+
+	void GatherSubdivisonsForTimespan(int32 StartTick, int32 EndTick, int32 QuantizationInterval, TArray<FGridInterval>& OutSubdivisions) const;
+
+	void GatherBarsForTimespan(int32 StartBar, int32 EndBar, int32 QuantizationInterval, TArray<FGridInterval>& OutBars) const;
+
+	void GetGridMetricsMusic(float InViewStartSeconds, float InViewEndSeconds, float PixelsPerSecond, float MinTickPx, float DesiredMajorTickPx, TArray<FGridInterval>& OutMajorIntervals, TArray<FGridInterval>& OutMinorIntervals) const;
+
+	void ComputeGridSpacingAtTick(int32 TickToCompute, float DesiredMajorTickSeconds, int32& OutMajorIntervalTicks) const;
+
+	void ComputeMinorGridSpacing(int32 StartTick, int32 EndTick, TArray<FGridInterval>& OutMinorIntervals) const;
+
+	TArray<FFrameBasedTempoPoint> GetTempoMetrics(float InViewStartSeconds, float InViewEndSeconds) const;
+
+	void GetTempoTimeSignatureMarkers(float InStartSeconds, float InEndSeconds, TArray<FTempoTimeSignatureMarker>& OutMarkers) const;
+
+private:
+	const FFrameBasedTimeSignaturePoint* GetTimeSignaturePointForTick(int32 Tick) const;
+	const FFrameBasedTimeSignaturePoint* GetTimeSignaturePointForBar(int32 Bar) const;
+	const FFrameBasedTimeSignaturePoint* GetTimeSignaturePointForFrame(double DecimalFrame) const;
+	const FFrameBasedTimeSignaturePoint* GetTimeSignaturePointForFrameTime(const FFrameTime& Frame) const;
+	const FFrameBasedTimeSignaturePoint* GetTimeSignaturePointForSeconds(double Seconds) const;
+
+	const FFrameBasedTempoPoint* GetTempoPointForTick(int32 Tick) const;
+	const FFrameBasedTempoPoint* GetTempoPointForFrame(double DecimalFrame) const;
+	const FFrameBasedTempoPoint* GetTempoPointForFrameTime(const FFrameTime& Frame) const;
+	const FFrameBasedTempoPoint* GetTempoPointForSeconds(double Seconds) const;
+
+	FString MakeBarBeatString(double DecimalFrame, const FFrameBasedTimeSignaturePoint* TimeSignature, const FFrameBasedTempoPoint* Tempo, FrameBasedMusicMap::EBarBeatFormat BarBeatFormat) const;
+};
+
+UINTERFACE(MinimalAPI)
+class UMusicMapSource : public UInterface
+{
+	GENERATED_BODY()
+};
+
+struct FMarkerProviderEntry
+{
+	FString Label;
+	FFrameNumber FrameNumber;
+
+	FMarkerProviderEntry(const FString& InLabel, FFrameNumber InFrameNumber)
+		: Label(InLabel)
+		, FrameNumber(InFrameNumber)
+	{}
+};
+
+USTRUCT()
+struct FMarkerProviderResults
+{
+	GENERATED_BODY()
+
+	struct FChannel
+	{
+		FString Name;
+		TArray<FMarkerProviderEntry> Markers;
+
+		FChannel(const FString& InName)
+			: Name(InName)
+		{}
+	};
+
+	bool IsEmpty() const { return Channels.IsEmpty(); }
+
+	TArray<FChannel> Channels;
+};
+
+class MUSICENVIRONMENT_API IMusicMapSource
+{
+public:
+
+	GENERATED_BODY()
+
+	UFUNCTION()
+	virtual void CreateFrameBasedMusicMap(UFrameBasedMusicMap* Map) const { return; }
+
+	UFUNCTION()
+	virtual FMarkerProviderResults GatherMarkers(const UFrameBasedMusicMap* Map) const { return FMarkerProviderResults(); }
+
+};
+
diff --git a/Engine/Plugins/Runtime/MusicEnvironment/Source/MusicEnvironment/Public/MusicClockSourceManager.h b/Engine/Plugins/Runtime/MusicEnvironment/Source/MusicEnvironment/Public/MusicClockSourceManager.h
new file mode 100644
--- /dev/null
+++ b/Engine/Plugins/Runtime/MusicEnvironment/Source/MusicEnvironment/Public/MusicClockSourceManager.h
@@ -0,0 +1,51 @@
+// Copyright Epic Games, Inc. All Rights Reserved.
+
+#pragma once
+
+#include "GameplayTagContainer.h"
+#include "MusicEnvironmentClockSource.h"
+#include "UObject/WeakInterfacePtr.h"
+
+#include "MusicClockSourceManager.generated.h"
+
+UCLASS(BlueprintType)
+class MUSICENVIRONMENT_API UMusicClockSourceManager : public UObject
+{
+	GENERATED_BODY()
+public:
+	
+	UFUNCTION(BlueprintCallable, Category="Music|Clock")
+	TScriptInterface<IMusicEnvironmentClockSource> FindClock(const FGameplayTag& Tag, bool bExactMatch = true);
+
+	UFUNCTION(BlueprintCallable, Category="Music|Clock")
+	void AddTaggedClock(const FGameplayTag& Tag, TScriptInterface<IMusicEnvironmentClockSource> InClock);
+	void AddTaggedClock(const FGameplayTag& Tag, UObject* InClock);
+
+	UFUNCTION(BlueprintCallable, Category="Music|Clock")
+	void RemoveTaggedClock(TScriptInterface<IMusicEnvironmentClockSource> InClock);
+	void RemoveTaggedClock(UObject* InClock);
+
+	UFUNCTION(BlueprintCallable, Category="Music|Clock")
+	void RemoveClockWithTag(const FGameplayTag& Tag);
+
+	UFUNCTION(BlueprintCallable, Category="Music|Clock")
+	TScriptInterface<IMusicEnvironmentClockSource> GetGlobalMusicClockAuthority();
+
+	UFUNCTION(BlueprintCallable, Category = "Music|Clock")
+	bool PushGlobalMusicClockAuthority(TScriptInterface<IMusicEnvironmentClockSource> InClock);
+	bool PushGlobalMusicClockAuthority(UObject* InClock);
+
+	UFUNCTION(BlueprintCallable, Category = "Music|Clock")
+	void RemoveGlobalClockAuthority(TScriptInterface<IMusicEnvironmentClockSource> InClock);
+	void RemoveGlobalClockAuthority(UObject* InClock);
+
+	UFUNCTION(BlueprintCallable, Category = "Music|Clock")
+	void PopMusicClockAuthority();
+
+private:
+	TMap<FGameplayTag, TWeakInterfacePtr<IMusicEnvironmentClockSource>> TaggedClocks;
+	TArray<TWeakInterfacePtr<IMusicEnvironmentClockSource>> GlobalClockSourceStack;
+
+	void CleanMap();
+	void CleanStack();
+};
diff --git a/Engine/Plugins/Runtime/MusicEnvironment/Source/MusicEnvironment/Public/MusicEnvironmentClockSource.h b/Engine/Plugins/Runtime/MusicEnvironment/Source/MusicEnvironment/Public/MusicEnvironmentClockSource.h
new file mode 100644
--- /dev/null
+++ b/Engine/Plugins/Runtime/MusicEnvironment/Source/MusicEnvironment/Public/MusicEnvironmentClockSource.h
@@ -0,0 +1,33 @@
+// Copyright Epic Games, Inc. All Rights Reserved.
+
+#pragma once
+
+#include "CoreMinimal.h"
+#include "FrameBasedMusicMap.h"
+#include "Misc/MusicalTime.h"
+#include "UObject/Interface.h"
+
+#include "MusicEnvironmentClockSource.generated.h"
+
+// This class does not need to be modified.
+UINTERFACE(BlueprintType, meta = (CannotImplementInterfaceInBlueprint))
+class MUSICENVIRONMENT_API UMusicEnvironmentClockSource : public UInterface
+{
+	GENERATED_BODY()
+};
+
+/**
+ * 
+ */
+class MUSICENVIRONMENT_API IMusicEnvironmentClockSource
+{
+	GENERATED_BODY()
+
+public:
+	virtual FMusicalTime GetMusicalTime() const = 0;
+	virtual int32 GetAbsoluteTickPosition() const = 0;
+	virtual FMusicalTime GetMusicalTimeWithSourceSpaceOffset(const FMusicalTime& Offset) const = 0;
+	virtual int32 GetAbsoluteTickPositionWithSourceSpaceOffset(const FMusicalTime& SourceOffset) const = 0;
+	virtual FMusicalTime Quantize(const FMusicalTime& MusicalTime, int32 QuantizationInterval, UFrameBasedMusicMap::EQuantizeDirection Direction = UFrameBasedMusicMap::EQuantizeDirection::Nearest) const = 0;
+	virtual bool CanAuditionInEditor() const = 0;
+};
diff --git a/Engine/Plugins/Runtime/MusicEnvironment/Source/MusicEnvironment/Public/MusicEnvironmentMetronome.h b/Engine/Plugins/Runtime/MusicEnvironment/Source/MusicEnvironment/Public/MusicEnvironmentMetronome.h
new file mode 100644
--- /dev/null
+++ b/Engine/Plugins/Runtime/MusicEnvironment/Source/MusicEnvironment/Public/MusicEnvironmentMetronome.h
@@ -0,0 +1,47 @@
+// Copyright Epic Games, Inc. All Rights Reserved.
+
+#pragma once
+
+#include "CoreMinimal.h"
+#include "FrameBasedMusicMap.h"
+#include "MusicEnvironmentClockSource.h"
+#include "UObject/Interface.h"
+
+#include "MusicEnvironmentMetronome.generated.h"
+
+// This class does not need to be modified.
+UINTERFACE()
+class UMusicEnvironmentMetronome : public UInterface
+{
+	GENERATED_BODY()
+};
+
+/**
+ * A music environment system that can spawn metronomes will have to return an 
+ * instance of a UObject that implements this interface from its "MovieSceneMetronomeSpawner". 
+ */
+class MUSICENVIRONMENT_API IMusicEnvironmentMetronome : public IMusicEnvironmentClockSource
+{
+	GENERATED_BODY()
+
+	// Add interface functions to this class. This is the class that will be inherited to implement this interface.
+public:
+	void SetMusicMap(UFrameBasedMusicMap* InMusicMap);
+
+	virtual bool Initialize(UWorld* InWorld) = 0;
+	virtual void Tick(float DeltaSeconds) = 0;
+	virtual void Start(double FromSeconds = 0.0) = 0;
+	virtual void Seek(double ToSeconds) = 0;
+	virtual void Stop() = 0;
+	virtual void Pause() = 0;
+	virtual void Resume() = 0;
+
+	bool SetTempo(float Bpm);
+	virtual float GetCurrentTempo() = 0;
+	virtual double GetCurrentPositionSeconds() = 0;
+	
+protected:
+	virtual void OnMusicMapSet() = 0;
+	virtual bool OnSetTempo(const float Bpm) = 0;
+	TStrongObjectPtr<UFrameBasedMusicMap> MusicMap;
+};
diff --git a/Engine/Plugins/Runtime/MusicEnvironment/Source/MusicEnvironment/Public/MusicEnvironmentModule.h b/Engine/Plugins/Runtime/MusicEnvironment/Source/MusicEnvironment/Public/MusicEnvironmentModule.h
new file mode 100644
--- /dev/null
+++ b/Engine/Plugins/Runtime/MusicEnvironment/Source/MusicEnvironment/Public/MusicEnvironmentModule.h
@@ -0,0 +1,21 @@
+// Copyright Epic Games, Inc. All Rights Reserved.
+
+#pragma once
+
+#include "CoreMinimal.h"
+#include "Modules/ModuleInterface.h"
+#include "Modules/ModuleManager.h"
+
+/**
+ * The public interface of the MusicEnvironment module
+ */
+class FMusicEnvironmentModule
+	: public IModuleInterface
+{
+public:
+
+	static inline FMusicEnvironmentModule& Get()
+	{
+		return FModuleManager::LoadModuleChecked<FMusicEnvironmentModule>("MusicEnvironment");
+	}
+};
diff --git a/Engine/Plugins/Runtime/MusicEnvironment/Source/MusicEnvironment/Public/MusicEnvironmentSubsystem.h b/Engine/Plugins/Runtime/MusicEnvironment/Source/MusicEnvironment/Public/MusicEnvironmentSubsystem.h
new file mode 100644
--- /dev/null
+++ b/Engine/Plugins/Runtime/MusicEnvironment/Source/MusicEnvironment/Public/MusicEnvironmentSubsystem.h
@@ -0,0 +1,44 @@
+// Copyright Epic Games, Inc. All Rights Reserved.
+
+#pragma once
+
+#include "CoreMinimal.h"
+#include "Subsystems/EngineSubsystem.h"
+#include "UObject/Class.h"
+#include "UObject/ObjectMacros.h"
+#include "UObject/UObjectGlobals.h"
+#include "UObject/WeakInterfacePtr.h"
+
+#include "MusicEnvironmentSubsystem.generated.h"
+
+// Forward Declarations
+class IMusicEnvironmentMetronome;
+class UMusicClockSourceManager;
+
+UCLASS(BlueprintType, Category = "Music", DisplayName = "MusicEnvironmentSubsystem")
+class MUSICENVIRONMENT_API UMusicEnvironmentSubsystem : public UEngineSubsystem
+{
+	GENERATED_BODY()
+public:
+	static UMusicEnvironmentSubsystem& Get();
+	
+	UFUNCTION(BlueprintCallable, Category = "Music|Clock")
+	UMusicClockSourceManager* GetClockSourceManager();
+
+	bool SetMetronomeClass(const UClass* InMetronomeType);
+
+	UFUNCTION(BlueprintCallable, Category = "Music|Clock")
+	bool CanSpawnMetronome() const;
+
+	UFUNCTION(BlueprintCallable, Category = "Music|Clock")
+	virtual TScriptInterface<IMusicEnvironmentMetronome> SpawnMetronome(UObject* Outer, FName Name = NAME_None);
+
+private:
+	UPROPERTY(Transient)
+	TObjectPtr<UMusicClockSourceManager> ClockSourceManager;
+
+	UPROPERTY(Transient)
+	TSoftClassPtr<UObject> MetronomeType;
+};
+
+MUSICENVIRONMENT_API DECLARE_LOG_CATEGORY_EXTERN(LogMusicEnvironment, Log, All);
diff --git a/Engine/Plugins/Runtime/MusicEnvironment/Source/MusicEnvironment/Public/MusicTypes/MusicalEnums.h b/Engine/Plugins/Runtime/MusicEnvironment/Source/MusicEnvironment/Public/MusicTypes/MusicalEnums.h
new file mode 100644
--- /dev/null
+++ b/Engine/Plugins/Runtime/MusicEnvironment/Source/MusicEnvironment/Public/MusicTypes/MusicalEnums.h
@@ -0,0 +1,24 @@
+// Copyright Epic Games, Inc. All Rights Reserved.
+
+#pragma once
+
+#include "MusicalEnums.generated.h"
+
+UENUM(BlueprintType)
+enum class EMusicalSyncPointCaptureMode : uint8
+{
+	DirectMapped              UMETA(ToolTip="Bar 1 | Beat 1 of the level sequence is mapped directly to Bar 1 | Beat 1 of the driving music clock. If needed the sequence is 'seeked' to match the current time of the music clock."),
+	Immediate                 UMETA(ToolTip="Bar 1 | Beat 1 of the level sequence will be mapped to the current position of the music clock. Now quantization is done. While the sequence will advance in a locked-step, tempo matched way, its offset will be maintained. CAREFUL! You probably want to use some sort of qunatized capture."),
+	NearestQuantizedPosition  UMETA(ToolTip="Bar 1 | Beat 1 of the level sequence is mapped to the NEAREST specified quantized musical subdivision. May be before or after the current musical time. Playback will start right away if the nearest subdivision is in the past, otherwise playback will start when the next subdivision arrives."),
+	NextQuantizedPosition     UMETA(ToolTip="Bar 1 | Beat 1 of the level sequence is mapped to the NEXT specified quantized musical subdivision. Will be on or after the current musical time. Playback will start when that subdivision arrives."),
+	PreviousQuantizedPosition UMETA(ToolTip="Bar 1 | Beat 1 of the level sequence is mapped to the Previous specified quantized musical subdivision. Will be on or before the current musical time. Playback will start right away."),
+};
+
+UENUM(BlueprintType)
+enum class EMusicalTimeSignatureRemapMode : uint8
+{
+	None                    UMETA(ToolTip="No attempt is made to maintain 'bar synchronizztion' if there are mismatched time signatured."),
+	ScaleBars               UMETA(ToolTip="Scale each bar as neccessary to match bar lengths in music clock."),
+	CropOrLoopBars          UMETA(ToolTip="Crop or loop bars as nessary to align downbeats between the level sequence and the music clock."),
+	ScaleBeatsAndCropOrLoop UMETA(ToolTip="Scale beats to match music clock beat lengths, and THEN crop or loop bars as nessary to align downbeats between the level sequence and the music clock.")
+};
diff --git a/Engine/Source/Runtime/AudioMixer/AudioMixer.Build.cs b/Engine/Source/Runtime/AudioMixer/AudioMixer.Build.cs
--- a/Engine/Source/Runtime/AudioMixer/AudioMixer.Build.cs
+++ b/Engine/Source/Runtime/AudioMixer/AudioMixer.Build.cs
@@ -32,7 +32,6 @@ namespace UnrealBuildTool.Rules
 					"AudioExtensions",
 					"AudioLinkCore",
 					"HeadMountedDisplay",
-					"MusicEnvironment",
 					"TraceLog"
 				}
 			);
diff --git a/Engine/Source/Runtime/AudioMixer/Private/AudioMixerModule.cpp b/Engine/Source/Runtime/AudioMixer/Private/AudioMixerModule.cpp
--- a/Engine/Source/Runtime/AudioMixer/Private/AudioMixerModule.cpp
+++ b/Engine/Source/Runtime/AudioMixer/Private/AudioMixerModule.cpp
@@ -10,7 +10,6 @@ public:
 	virtual void StartupModule() override
 	{
 		FModuleManager::Get().LoadModuleChecked(TEXT("AudioMixerCore"));
-		FModuleManager::Get().LoadModuleChecked(TEXT("MusicEnvironment"));
 		FModuleManager::Get().LoadModuleChecked(TEXT("SignalProcessing"));
 	}
 };
diff --git a/Engine/Source/Runtime/MusicEnvironment/MusicEnvironment.Build.cs b/Engine/Source/Runtime/MusicEnvironment/MusicEnvironment.Build.cs
deleted file mode 100644
diff --git a/Engine/Source/Runtime/MusicEnvironment/Private/FrameBasedMusicMap.cpp b/Engine/Source/Runtime/MusicEnvironment/Private/FrameBasedMusicMap.cpp
deleted file mode 100644
diff --git a/Engine/Source/Runtime/MusicEnvironment/Private/MusicClockSourceManager.cpp b/Engine/Source/Runtime/MusicEnvironment/Private/MusicClockSourceManager.cpp
deleted file mode 100644
diff --git a/Engine/Source/Runtime/MusicEnvironment/Private/MusicEnvironmentMetronome.cpp b/Engine/Source/Runtime/MusicEnvironment/Private/MusicEnvironmentMetronome.cpp
deleted file mode 100644
diff --git a/Engine/Source/Runtime/MusicEnvironment/Private/MusicEnvironmentModule.cpp b/Engine/Source/Runtime/MusicEnvironment/Private/MusicEnvironmentModule.cpp
deleted file mode 100644
diff --git a/Engine/Source/Runtime/MusicEnvironment/Private/MusicEnvironmentSubsystem.cpp b/Engine/Source/Runtime/MusicEnvironment/Private/MusicEnvironmentSubsystem.cpp
deleted file mode 100644
diff --git a/Engine/Source/Runtime/MusicEnvironment/Public/FrameBasedMusicMap.h b/Engine/Source/Runtime/MusicEnvironment/Public/FrameBasedMusicMap.h
deleted file mode 100644
diff --git a/Engine/Source/Runtime/MusicEnvironment/Public/MusicClockSourceManager.h b/Engine/Source/Runtime/MusicEnvironment/Public/MusicClockSourceManager.h
deleted file mode 100644
diff --git a/Engine/Source/Runtime/MusicEnvironment/Public/MusicEnvironmentClockSource.h b/Engine/Source/Runtime/MusicEnvironment/Public/MusicEnvironmentClockSource.h
deleted file mode 100644
diff --git a/Engine/Source/Runtime/MusicEnvironment/Public/MusicEnvironmentMetronome.h b/Engine/Source/Runtime/MusicEnvironment/Public/MusicEnvironmentMetronome.h
deleted file mode 100644
diff --git a/Engine/Source/Runtime/MusicEnvironment/Public/MusicEnvironmentModule.h b/Engine/Source/Runtime/MusicEnvironment/Public/MusicEnvironmentModule.h
deleted file mode 100644
diff --git a/Engine/Source/Runtime/MusicEnvironment/Public/MusicEnvironmentSubsystem.h b/Engine/Source/Runtime/MusicEnvironment/Public/MusicEnvironmentSubsystem.h
deleted file mode 100644
diff --git a/Engine/Source/Runtime/MusicEnvironment/Public/MusicTypes/MusicalEnums.h b/Engine/Source/Runtime/MusicEnvironment/Public/MusicTypes/MusicalEnums.h
deleted file mode 100644
