Fix base class of UMusicEnvironmentMetronome to assure proper UClass reflection.

diff --git a/Engine/Plugins/Runtime/MusicEnvironment/Source/MusicEnvironment/Public/MusicEnvironmentMetronome.h b/Engine/Plugins/Runtime/MusicEnvironment/Source/MusicEnvironment/Public/MusicEnvironmentMetronome.h
--- a/Engine/Plugins/Runtime/MusicEnvironment/Source/MusicEnvironment/Public/MusicEnvironmentMetronome.h
+++ b/Engine/Plugins/Runtime/MusicEnvironment/Source/MusicEnvironment/Public/MusicEnvironmentMetronome.h
@@ -11,7 +11,7 @@
 
 // This class does not need to be modified.
 UINTERFACE()
-class UMusicEnvironmentMetronome : public UInterface
+class UMusicEnvironmentMetronome : public UMusicEnvironmentClockSource
 {
 	GENERATED_BODY()
 };
