FrameBasedMusicMap

Use OldMap to determine Seconds -> Ticks -> MusicTime when Remapping Time from Old Map to New Map

#rnx
#tests

diff --git a/Engine/Plugins/Runtime/MusicEnvironment/Source/MusicEnvironment/Private/FrameBasedMusicMap.cpp b/Engine/Plugins/Runtime/MusicEnvironment/Source/MusicEnvironment/Private/FrameBasedMusicMap.cpp
--- a/Engine/Plugins/Runtime/MusicEnvironment/Source/MusicEnvironment/Private/FrameBasedMusicMap.cpp
+++ b/Engine/Plugins/Runtime/MusicEnvironment/Source/MusicEnvironment/Private/FrameBasedMusicMap.cpp
@@ -437,11 +437,24 @@ FFrameTime UFrameBasedMusicMap::RemapFrameTime(FFrameTime SourceFrameTime, const
 {
 	if (IsEmpty())
 	{
+		UE_LOG(LogMusicMap, Warning, TEXT("Unable to remap times, MusicMap is empty"));
+		return SourceFrameTime;
+	}
+
+	if (OldMap == nullptr)
+	{
+		UE_LOG(LogMusicMap, Warning, TEXT("Unable to remap times, Old MusicMap is null"));
 		return SourceFrameTime;
 	}
 	
-	float Seconds = FrameResolution.AsSeconds(SourceFrameTime);
-	const FFrameBasedTempoPoint* TempoPoint = GetTempoPointForSeconds(Seconds);
+	if (OldMap->IsEmpty())
+	{
+		UE_LOG(LogMusicMap, Warning, TEXT("Unable to remap times, Old MusicMap is empty"));
+		return SourceFrameTime;
+	}
+	
+	float Seconds = OldMap->FrameResolution.AsSeconds(SourceFrameTime);
+	const FFrameBasedTempoPoint* TempoPoint = OldMap->GetTempoPointForSeconds(Seconds);
 	int32 Tick = TempoPoint->SecondsToTick(Seconds);
 
 	FMusicalTime MusicalTime = OldMap->GetMusicalTime(Tick);
