Fix some trivial spelling errors in tooltips for enums

diff --git a/Engine/Plugins/Runtime/MusicEnvironment/Source/MusicEnvironment/Public/MusicTypes/MusicalEnums.h b/Engine/Plugins/Runtime/MusicEnvironment/Source/MusicEnvironment/Public/MusicTypes/MusicalEnums.h
--- a/Engine/Plugins/Runtime/MusicEnvironment/Source/MusicEnvironment/Public/MusicTypes/MusicalEnums.h
+++ b/Engine/Plugins/Runtime/MusicEnvironment/Source/MusicEnvironment/Public/MusicTypes/MusicalEnums.h
@@ -17,8 +17,8 @@ enum class EMusicalSyncPointCaptureMode : uint8
 UENUM(BlueprintType)
 enum class EMusicalTimeSignatureRemapMode : uint8
 {
-	None                    UMETA(ToolTip="No attempt is made to maintain 'bar synchronizztion' if there are mismatched time signatured."),
-	ScaleBars               UMETA(ToolTip="Scale each bar as neccessary to match bar lengths in music clock."),
-	CropOrLoopBars          UMETA(ToolTip="Crop or loop bars as nessary to align downbeats between the level sequence and the music clock."),
-	ScaleBeatsAndCropOrLoop UMETA(ToolTip="Scale beats to match music clock beat lengths, and THEN crop or loop bars as nessary to align downbeats between the level sequence and the music clock.")
+	None                    UMETA(ToolTip="No attempt is made to maintain 'bar synchronization' if there are mismatched time signature."),
+	ScaleBars               UMETA(ToolTip="Scale each bar as necessary to match bar lengths in music clock."),
+	CropOrLoopBars          UMETA(ToolTip="Crop or loop bars as necessary to align downbeats between the level sequence and the music clock."),
+	ScaleBeatsAndCropOrLoop UMETA(ToolTip="Scale beats to match music clock beat lengths, and THEN crop or loop bars as necessary to align downbeats between the level sequence and the music clock.")
 };
