- Tweaks to the MusicEnvironmentClockSource interface
- Addition of MusicalAsset and MusicHandle interfaces

diff --git a/Engine/Plugins/Runtime/Harmonix/Source/HarmonixMetasound/Private/HarmonixMetasound/Components/MusicClockComponent.cpp b/Engine/Plugins/Runtime/Harmonix/Source/HarmonixMetasound/Private/HarmonixMetasound/Components/MusicClockComponent.cpp
--- a/Engine/Plugins/Runtime/Harmonix/Source/HarmonixMetasound/Private/HarmonixMetasound/Components/MusicClockComponent.cpp
+++ b/Engine/Plugins/Runtime/Harmonix/Source/HarmonixMetasound/Private/HarmonixMetasound/Components/MusicClockComponent.cpp
@@ -806,7 +806,13 @@ const FMidiSongPos& UMusicClockComponent::GetPreviousSongPosInternal(ECalibrated
 	return PosZero;
 }
 
-FMusicalTime UMusicClockComponent::GetMusicalTime() const
+float UMusicClockComponent::GetPositionSeconds() const
+{
+	const FMidiSongPos& SongPos = ClockDriver->GetCurrentSongPos(ECalibratedMusicTimebase::VideoRenderTime);
+	return SongPos.SecondsIncludingCountIn;
+}
+
+FMusicalTime UMusicClockComponent::GetPositionMusicalTime() const
 {
 	FMusicalTime Result;
 	if (!ClockDriver)
@@ -827,7 +833,7 @@ FMusicalTime UMusicClockComponent::GetMusicalTime() const
 	return Result;
 }
 
-int32 UMusicClockComponent::GetAbsoluteTickPosition() const
+int32 UMusicClockComponent::GetPositionAbsoluteTick() const
 {
 	// Note: Since this is an override of a function in the music environment system we need to
 	// be sure we are using its standard ticks per quarter note units.
@@ -837,13 +843,13 @@ int32 UMusicClockComponent::GetAbsoluteTickPosition() const
 }
 
 
-FMusicalTime UMusicClockComponent::GetMusicalTimeWithSourceSpaceOffset(const FMusicalTime& Offset) const
+FMusicalTime UMusicClockComponent::GetPositionMusicalTime(const FMusicalTime& SourceSpaceOffset) const
 {
 	FMusicalTime Result;
 	const ISongMapEvaluator& MapEvaluator = GetSongMaps();
 	float CurrentTick = MapEvaluator.FractionalBarIncludingCountInToTick(GetCurrentSongPosInternal(ECalibratedMusicTimebase::VideoRenderTime).BarsIncludingCountIn);
 	
-	float OffsetTick = MapEvaluator.FractionalBarIncludingCountInToTick(static_cast<float>(Offset.FractionalBar()));
+	float OffsetTick = MapEvaluator.FractionalBarIncludingCountInToTick(static_cast<float>(SourceSpaceOffset.FractionalBar()));
 	float AdjustedTick = CurrentTick - OffsetTick;
 
 	const FTimeSignature* TimeSignature = MapEvaluator.GetTimeSignatureAtTick(FMath::FloorToInt32(AdjustedTick));
@@ -855,13 +861,13 @@ FMusicalTime UMusicClockComponent::GetMusicalTimeWithSourceSpaceOffset(const FMu
 	return Result;
 }
 
-int32 UMusicClockComponent::GetAbsoluteTickPositionWithSourceSpaceOffset(const FMusicalTime& Offset) const
+int32 UMusicClockComponent::GetPositionAbsoluteTick(const FMusicalTime& SourceSpaceOffset) const
 {
 	// Note: Since this is an override of a function in the music environment system we need to
 	// be sure we are using its standard ticks per quarter note units.
 	const ISongMapEvaluator& MapEvaluator = GetSongMaps();
 	float CurrentTick = MapEvaluator.FractionalBarIncludingCountInToTick(GetCurrentSongPosInternal(ECalibratedMusicTimebase::VideoRenderTime).BarsIncludingCountIn);
-	float OffsetTick = MapEvaluator.FractionalBarIncludingCountInToTick(static_cast<float>(Offset.FractionalBar()));
+	float OffsetTick = MapEvaluator.FractionalBarIncludingCountInToTick(static_cast<float>(SourceSpaceOffset.FractionalBar()));
 	float AdjustedTick = CurrentTick - OffsetTick;
 	return FMath::FloorToInt32((AdjustedTick / static_cast<float>(MapEvaluator.GetTicksPerQuarterNote())) * static_cast<float>(MusicalTime::TicksPerQuarterNote));
 }
diff --git a/Engine/Plugins/Runtime/Harmonix/Source/HarmonixMetasound/Private/HarmonixMetasound/MusicEnvironmentSupport/HarmonixMusicEnvironmentMetronome.cpp b/Engine/Plugins/Runtime/Harmonix/Source/HarmonixMetasound/Private/HarmonixMetasound/MusicEnvironmentSupport/HarmonixMusicEnvironmentMetronome.cpp
--- a/Engine/Plugins/Runtime/Harmonix/Source/HarmonixMetasound/Private/HarmonixMetasound/MusicEnvironmentSupport/HarmonixMusicEnvironmentMetronome.cpp
+++ b/Engine/Plugins/Runtime/Harmonix/Source/HarmonixMetasound/Private/HarmonixMetasound/MusicEnvironmentSupport/HarmonixMusicEnvironmentMetronome.cpp
@@ -48,6 +48,7 @@ void UHarmonixMusicEnvironmentMetronome::Start(double FromSeconds)
 	}
 
 	RebuildMidiFile();
+	SecondsWhenStarted = FromSeconds;
 	MusicalTimeWhenStarted = MidiFile->GetSongMaps()->GetMusicalTimeAtSeconds(FromSeconds);
 	FAudioParameter MidiFileParam(TEXT("MIDI File"), MidiFile);
 	FAudioParameter SeekTargetMs(TEXT("SeekTarget"), static_cast<float>(FromSeconds * 1000.0f));
@@ -119,41 +120,51 @@ double UHarmonixMusicEnvironmentMetronome::GetCurrentPositionSeconds() const
 	return 0.0;
 }
 
-FMusicalTime UHarmonixMusicEnvironmentMetronome::GetMusicalTime() const
+float UHarmonixMusicEnvironmentMetronome::GetPositionSeconds() const
+{
+	float TimeSeconds = (float)SecondsWhenStarted;
+	if (MusicClockComponent)
+	{
+		return MusicClockComponent->GetPositionSeconds() < TimeSeconds ? TimeSeconds : MusicClockComponent->GetPositionSeconds();
+	}
+	return TimeSeconds;
+}
+
+FMusicalTime UHarmonixMusicEnvironmentMetronome::GetPositionMusicalTime() const
 {
 	FMusicalTime CurrentMusicalTime;
 	if (MusicClockComponent)
 	{
-		CurrentMusicalTime = MusicClockComponent->GetMusicalTime();
+		CurrentMusicalTime = MusicClockComponent->GetPositionMusicalTime();
 		return MusicalTimeWhenStarted > CurrentMusicalTime ? MusicalTimeWhenStarted : CurrentMusicalTime;
 	}
 	return CurrentMusicalTime;
 }
 
-int32 UHarmonixMusicEnvironmentMetronome::GetAbsoluteTickPosition() const
+int32 UHarmonixMusicEnvironmentMetronome::GetPositionAbsoluteTick() const
 {
 	if (MusicClockComponent)
 	{
-		return MusicClockComponent->GetAbsoluteTickPosition();
+		return MusicClockComponent->GetPositionAbsoluteTick();
 	}
 	return 0;
 }
 
-FMusicalTime UHarmonixMusicEnvironmentMetronome::GetMusicalTimeWithSourceSpaceOffset(const FMusicalTime& Offset) const
+FMusicalTime UHarmonixMusicEnvironmentMetronome::GetPositionMusicalTime(const FMusicalTime& SourceSpaceOffset) const
 {
 	FMusicalTime MusicalTime;
 	if (MusicClockComponent)
 	{
-		return MusicClockComponent->GetMusicalTimeWithSourceSpaceOffset(Offset);
+		return MusicClockComponent->GetPositionMusicalTime(SourceSpaceOffset);
 	}
 	return MusicalTime;
 }
 
-int32 UHarmonixMusicEnvironmentMetronome::GetAbsoluteTickPositionWithSourceSpaceOffset(const FMusicalTime& Offset) const
+int32 UHarmonixMusicEnvironmentMetronome::GetPositionAbsoluteTick(const FMusicalTime& SourceSpaceOffset) const
 {
 	if (MusicClockComponent)
 	{
-		return MusicClockComponent->GetAbsoluteTickPositionWithSourceSpaceOffset(Offset);
+		return MusicClockComponent->GetPositionAbsoluteTick(SourceSpaceOffset);
 	}
 	return 0;
 }
diff --git a/Engine/Plugins/Runtime/Harmonix/Source/HarmonixMetasound/Private/HarmonixMetasound/MusicEnvironmentSupport/HarmonixMusicEnvironmentMetronome.h b/Engine/Plugins/Runtime/Harmonix/Source/HarmonixMetasound/Private/HarmonixMetasound/MusicEnvironmentSupport/HarmonixMusicEnvironmentMetronome.h
--- a/Engine/Plugins/Runtime/Harmonix/Source/HarmonixMetasound/Private/HarmonixMetasound/MusicEnvironmentSupport/HarmonixMusicEnvironmentMetronome.h
+++ b/Engine/Plugins/Runtime/Harmonix/Source/HarmonixMetasound/Private/HarmonixMetasound/MusicEnvironmentSupport/HarmonixMusicEnvironmentMetronome.h
@@ -43,12 +43,13 @@ public:
 	UE_API virtual bool IsMuted() const override;
 
 	// vvv BEGIN  IMusicEnvironmentClockSource Implementation vvv
-	UE_API virtual FMusicalTime GetMusicalTime() const override;
-	UE_API virtual int32 GetAbsoluteTickPosition() const override;
-	UE_API virtual FMusicalTime GetMusicalTimeWithSourceSpaceOffset(const FMusicalTime& Offset) const override;
-	UE_API virtual int32 GetAbsoluteTickPositionWithSourceSpaceOffset(const FMusicalTime& Offset) const override;
+	UE_API virtual float GetPositionSeconds() const override;
+	UE_API virtual FMusicalTime GetPositionMusicalTime() const override;
+	UE_API virtual FMusicalTime GetPositionMusicalTime(const FMusicalTime& SourceSpaceOffset) const override;
+	UE_API virtual int32 GetPositionAbsoluteTick() const override;
+	UE_API virtual int32 GetPositionAbsoluteTick(const FMusicalTime& SourceSpaceOffset) const override;
 	UE_API virtual FMusicalTime Quantize(const FMusicalTime& MusicalTime, int32 QuantizationInterval, UFrameBasedMusicMap::EQuantizeDirection Direction = UFrameBasedMusicMap::EQuantizeDirection::Nearest) const override;
-	virtual bool CanAuditionInEditor() const override { return true; }
+	UE_API virtual bool CanAuditionInEditor() const override { return true; }
 	// ^^^ END IMusicEnvironmentClockSource Implementation ^^^
 
 	UE_API virtual void BeginDestroy() override;
@@ -83,6 +84,7 @@ private:
 	float CurrentSpeed = 1.0f;
 	bool bIsMuted = false;
 	float CurrentVolume = 1.0f;
+	double SecondsWhenStarted = 0.0;
 	FMusicalTime MusicalTimeWhenStarted;
 	
 	UE_API bool BuildAndStartMetasound(UWorld* InWorld);
diff --git a/Engine/Plugins/Runtime/Harmonix/Source/HarmonixMetasound/Public/HarmonixMetasound/Components/MusicClockComponent.h b/Engine/Plugins/Runtime/Harmonix/Source/HarmonixMetasound/Public/HarmonixMetasound/Components/MusicClockComponent.h
--- a/Engine/Plugins/Runtime/Harmonix/Source/HarmonixMetasound/Public/HarmonixMetasound/Components/MusicClockComponent.h
+++ b/Engine/Plugins/Runtime/Harmonix/Source/HarmonixMetasound/Public/HarmonixMetasound/Components/MusicClockComponent.h
@@ -466,12 +466,13 @@ public:
 	UE_API const FMidiSongPos& GetRawUnsmoothedAudioRenderPos() const;
 
 	// vvv BEGIN  IMusicEnvironmentClockSource Implementation vvv
-	UE_API virtual FMusicalTime GetMusicalTime() const override;
-	UE_API virtual int32 GetAbsoluteTickPosition() const override;
-	UE_API virtual FMusicalTime GetMusicalTimeWithSourceSpaceOffset(const FMusicalTime& Offset) const override;
-	UE_API virtual int32 GetAbsoluteTickPositionWithSourceSpaceOffset(const FMusicalTime& Offset) const override;
+	UE_API virtual float GetPositionSeconds() const override;
+	UE_API virtual FMusicalTime GetPositionMusicalTime() const override;
+	UE_API virtual FMusicalTime GetPositionMusicalTime(const FMusicalTime& SourceSpaceOOffset) const override;
+	UE_API virtual int32 GetPositionAbsoluteTick() const override;
+	UE_API virtual int32 GetPositionAbsoluteTick(const FMusicalTime& SourceSpaceOffset) const override;
 	UE_API virtual FMusicalTime Quantize(const FMusicalTime& MusicalTime, int32 QuantizationInterval, UFrameBasedMusicMap::EQuantizeDirection Direction = UFrameBasedMusicMap::EQuantizeDirection::Nearest) const;
-	virtual bool CanAuditionInEditor() const override { return false; }
+	UE_API virtual bool CanAuditionInEditor() const override { return false; }
 	// ^^^ END IMusicEnvironmentClockSource Implementation ^^^
 
 public:
diff --git a/Engine/Plugins/Runtime/MusicEnvironment/Source/MusicEnvironment/Private/MusicEnvironmentClockSource.cpp b/Engine/Plugins/Runtime/MusicEnvironment/Source/MusicEnvironment/Private/MusicEnvironmentClockSource.cpp
new file mode 100644
--- /dev/null
+++ b/Engine/Plugins/Runtime/MusicEnvironment/Source/MusicEnvironment/Private/MusicEnvironmentClockSource.cpp
@@ -0,0 +1,10 @@
+// Copyright Epic Games, Inc. All Rights Reserved.
+
+#include "MusicEnvironmentClockSource.h"
+
+void IMusicEnvironmentClockSource::GetCurrentBarBeat(float& Bar, float& BeatInBar)
+{
+	FMusicalTime MusicalTime = GetPositionMusicalTime();
+	Bar = MusicalTime.FractionalBar();
+	BeatInBar = MusicalTime.FractionalBeatInBar();
+}
diff --git a/Engine/Plugins/Runtime/MusicEnvironment/Source/MusicEnvironment/Private/MusicTypes/MusicHandle.cpp b/Engine/Plugins/Runtime/MusicEnvironment/Source/MusicEnvironment/Private/MusicTypes/MusicHandle.cpp
new file mode 100644
--- /dev/null
+++ b/Engine/Plugins/Runtime/MusicEnvironment/Source/MusicEnvironment/Private/MusicTypes/MusicHandle.cpp
@@ -0,0 +1,93 @@
+// Copyright Epic Games, Inc. All Rights Reserved.
+
+#include "MusicTypes/MusicHandle.h"
+#include "MusicClockSourceManager.h"
+#include "MusicEnvironmentSubsystem.h"
+
+void IMusicHandle::Stop()
+{
+	RelinquishGlobalMusicClockAuthority();
+	UnregisterAsTaggedClock();
+	Stop_Internal();
+}
+
+void IMusicHandle::Kill()
+{
+	RelinquishGlobalMusicClockAuthority();
+	UnregisterAsTaggedClock();
+	Kill_Internal();
+}
+
+void IMusicHandle::BranchOnTransportState(EMusicHanldeTransportState& Branches) const
+{
+	Branches = GetTransportState();
+}
+
+void IMusicHandle::GetCurrentBarBeat(float& Bar, float& BeatInBar, EMusicHanldeClockValidity& Branches)
+{
+	TScriptInterface<IMusicEnvironmentClockSource> MusicClockSource = GetMusicClockSource();
+	if (MusicClockSource)
+	{
+		MusicClockSource->GetCurrentBarBeat(Bar, BeatInBar);
+		Branches = EMusicHanldeClockValidity::ClockValid;
+	}
+	else
+	{
+		Bar = 0.0f;
+		BeatInBar = 0.0f;
+		Branches = EMusicHanldeClockValidity::ClockInvalid;
+	}
+}
+
+void IMusicHandle::BecomeGlobalMusicClockAuthority()
+{
+	UMusicEnvironmentSubsystem::Get().GetClockSourceManager()->PushGlobalMusicClockAuthority(GetMusicClockSource());
+	IsOnGlobalClockAuthorityStack = true;
+}
+
+void IMusicHandle::RelinquishGlobalMusicClockAuthority()
+{
+	if (IsOnGlobalClockAuthorityStack)
+	{
+		UMusicEnvironmentSubsystem::Get().GetClockSourceManager()->RemoveGlobalClockAuthority(GetMusicClockSource());
+		IsOnGlobalClockAuthorityStack = false;
+	}
+}
+
+void IMusicHandle::RegisterAsTaggedClock(const FGameplayTag& RegisterAsTagged)
+{
+	if (RegisterAsTagged.IsValid())
+	{
+		UMusicEnvironmentSubsystem::Get().GetClockSourceManager()->AddTaggedClock(RegisterAsTagged, GetMusicClockSource());
+		RegisteredToTags.AddTag(RegisterAsTagged);
+	}
+}
+
+void IMusicHandle::UnregisterAsTaggedClock(const FGameplayTag& RegisterAsTagged)
+{
+	if (!RegisterAsTagged.IsValid())
+	{
+		// remove all registrations...
+		UMusicEnvironmentSubsystem::Get().GetClockSourceManager()->RemoveTaggedClock(GetMusicClockSource());
+		RegisteredToTags.Reset();
+	}
+	else if (RegisteredToTags.HasTag(RegisterAsTagged))
+	{
+		UMusicEnvironmentSubsystem::Get().GetClockSourceManager()->RemoveClockWithTag(RegisterAsTagged);
+		RegisteredToTags.RemoveTag(RegisterAsTagged);
+	}
+}
+
+void UMusicHandleBlueprintHelpers::BranchOnTransportState(TScriptInterface<IMusicHandle> Handle, EMusicHanldeTransportState& Branches)
+{
+	Branches = UMusicHandleBlueprintHelpers::GetTransportState(Handle);
+}
+
+EMusicHanldeTransportState UMusicHandleBlueprintHelpers::GetTransportState(TScriptInterface<IMusicHandle> Handle)
+{
+	if (!Handle)
+	{
+		return EMusicHanldeTransportState::Stale;
+	}
+	return Handle->GetTransportState();
+}
diff --git a/Engine/Plugins/Runtime/MusicEnvironment/Source/MusicEnvironment/Private/MusicTypes/MusicalAsset.cpp b/Engine/Plugins/Runtime/MusicEnvironment/Source/MusicEnvironment/Private/MusicTypes/MusicalAsset.cpp
new file mode 100644
--- /dev/null
+++ b/Engine/Plugins/Runtime/MusicEnvironment/Source/MusicEnvironment/Private/MusicTypes/MusicalAsset.cpp
@@ -0,0 +1,43 @@
+// Copyright Epic Games, Inc. All Rights Reserved.
+
+#include "MusicTypes/MusicalAsset.h"
+#include "MusicTypes/MusicHandle.h"
+
+
+TScriptInterface<IMusicHandle> IMusicalAsset::PrepareToPlay(UObject* PlaybackContext, UAudioComponent* OnComponent,
+	float FromSeconds, bool BecomeAuthoritativeClock, const FGameplayTag RegisterAsTaggedClock, bool IsAudition)
+{
+	TScriptInterface<IMusicHandle> Handle = PrepareToPlay_Internal(PlaybackContext, OnComponent, FromSeconds, IsAudition);
+	if (Handle)
+	{
+		if (BecomeAuthoritativeClock)
+		{
+			Handle->BecomeGlobalMusicClockAuthority();
+		}
+
+		if (RegisterAsTaggedClock.IsValid())
+		{
+			Handle->RegisterAsTaggedClock(RegisterAsTaggedClock);
+		}
+	}
+	return Handle;
+}
+
+TScriptInterface<IMusicHandle> IMusicalAsset::Play(UObject* PlaybackContext, UAudioComponent* OnComponent,
+	float FromSeconds, bool BecomeAuthoritativeClock, const FGameplayTag RegisterAsTaggedClock, bool IsAudition)
+{
+	TScriptInterface<IMusicHandle> Handle = Play_Internal(PlaybackContext, OnComponent, FromSeconds, IsAudition);
+	if (Handle)
+	{
+		if (BecomeAuthoritativeClock)
+		{
+			Handle->BecomeGlobalMusicClockAuthority();
+		}
+
+		if (RegisterAsTaggedClock.IsValid())
+		{
+			Handle->RegisterAsTaggedClock(RegisterAsTaggedClock);
+		}
+	}
+	return Handle;
+}
\ No newline at end of file
diff --git a/Engine/Plugins/Runtime/MusicEnvironment/Source/MusicEnvironment/Public/MusicEnvironmentClockSource.h b/Engine/Plugins/Runtime/MusicEnvironment/Source/MusicEnvironment/Public/MusicEnvironmentClockSource.h
--- a/Engine/Plugins/Runtime/MusicEnvironment/Source/MusicEnvironment/Public/MusicEnvironmentClockSource.h
+++ b/Engine/Plugins/Runtime/MusicEnvironment/Source/MusicEnvironment/Public/MusicEnvironmentClockSource.h
@@ -9,6 +9,8 @@
 
 #include "MusicEnvironmentClockSource.generated.h"
 
+#define UE_API MUSICENVIRONMENT_API
+
 // This class does not need to be modified.
 UINTERFACE(MinimalAPI, BlueprintType, meta = (CannotImplementInterfaceInBlueprint))
 class UMusicEnvironmentClockSource : public UInterface
@@ -24,10 +26,17 @@ class IMusicEnvironmentClockSource
 	GENERATED_BODY()
 
 public:
-	virtual FMusicalTime GetMusicalTime() const = 0;
-	virtual int32 GetAbsoluteTickPosition() const = 0;
-	virtual FMusicalTime GetMusicalTimeWithSourceSpaceOffset(const FMusicalTime& Offset) const = 0;
-	virtual int32 GetAbsoluteTickPositionWithSourceSpaceOffset(const FMusicalTime& SourceOffset) const = 0;
+
+	UFUNCTION(BlueprintCallable, Category = "Music|Clock")
+	UE_API virtual void GetCurrentBarBeat(float& Bar, float& BeatInBar);
+
+	virtual float GetPositionSeconds() const = 0;
+	virtual FMusicalTime GetPositionMusicalTime() const = 0;
+	virtual FMusicalTime GetPositionMusicalTime(const FMusicalTime& SourceSpaceOffset) const = 0;
+	virtual int32 GetPositionAbsoluteTick() const = 0;
+	virtual int32 GetPositionAbsoluteTick(const FMusicalTime& SourceSpaceOffset) const = 0;
 	virtual FMusicalTime Quantize(const FMusicalTime& MusicalTime, int32 QuantizationInterval, UFrameBasedMusicMap::EQuantizeDirection Direction = UFrameBasedMusicMap::EQuantizeDirection::Nearest) const = 0;
 	virtual bool CanAuditionInEditor() const = 0;
 };
+
+#undef UE_API // MUSICENVIRONMENT_API
diff --git a/Engine/Plugins/Runtime/MusicEnvironment/Source/MusicEnvironment/Public/MusicTypes/MusicHandle.h b/Engine/Plugins/Runtime/MusicEnvironment/Source/MusicEnvironment/Public/MusicTypes/MusicHandle.h
new file mode 100644
--- /dev/null
+++ b/Engine/Plugins/Runtime/MusicEnvironment/Source/MusicEnvironment/Public/MusicTypes/MusicHandle.h
@@ -0,0 +1,225 @@
+// Copyright Epic Games, Inc. All Rights Reserved.
+
+#pragma once
+
+#include "MusicalAsset.h"
+#include "MusicEnvironmentClockSource.h"
+#include "UObject/Interface.h"
+
+#include "MusicHandle.generated.h"
+
+#define UE_API MUSICENVIRONMENT_API
+
+/**
+ * This file defines:
+ * ==================
+ * UMusicHandle: The UInterface derived class needed to support Unreal's interface system. This 
+ *               class definition stays empty.
+ * 
+ * IMusicHandle: The actual pure virtual definition of the interface. This is what 
+ *               derived classes will inherit from. AND... IMusicalAssets return 
+ *               TScriptInterface<IMusicHandle> when they are asked to instantiate themselves.
+ *               BUT NOTE... If a non-UCLASS or a non-USTRUCT asks an IMusicalAsset to instantiate 
+ *               itself there may not be anything holding a strong reference to the object the
+ *               script interface is pointing to. More on this below.
+ * 
+ * FMusicHandle: This is a convenience type definition that amounts to a TScriptInterface<IMusicHandle>. 
+ *               ULASSS and USTRUCT things can hold one of these as a member and mark it as a UPROPERTY()
+ *               so the object that implements the script interface will not be garbage collected.
+ *               It looks cleaner.
+ * 
+ * FStrongMusicHandle: This wraps a TScriptInterface<IMusicHandle> to hold a TStrongObjectPtr to 
+ *               the UObject that implements the interface. It is like having a TStrongInterfacePtr<>,
+ *               a type that doesn't exist in the engine. This allows non-UCLASS, non-USTRUCT things to
+ *               hold music handles, and for the UStuff inside of those handles to not be garbage collected
+ *               until the FStrongMusicHandle is destroyed.
+ */
+
+class UAudioComponent;
+
+UENUM(BlueprintType)
+enum class EMusicHanldeClockValidity : uint8
+{
+	ClockValid,
+	ClockInvalid
+};
+
+UENUM(BlueprintType)
+enum class EMusicHanldeTransportState : uint8
+{
+	PreparingToPlay,
+	ReadyToPlay,
+	Playing,
+	Paused,
+	Stale,
+};
+
+// This class does not need to be modified.
+UINTERFACE(BlueprintType, meta = (CannotImplementInterfaceInBlueprint))
+class UE_API UMusicHandle : public UInterface
+{
+	GENERATED_BODY()
+};
+
+class UE_API IMusicHandle
+{
+public:
+	GENERATED_BODY()
+
+	UFUNCTION(BlueprintCallable, Category = "Music")
+	virtual bool IsValid() const = 0;
+
+	virtual void Tick(float DeltaSeconds) = 0;
+
+	//** Transport
+	UFUNCTION(BlueprintCallable, Category = "Music")
+	virtual bool IsReadyToPlay(float FromSeconds = 0.0f) { return true; }
+	UFUNCTION(BlueprintCallable, Category = "Music")
+	virtual bool Play(float FromSeconds = 0.0f) = 0;
+	UFUNCTION(BlueprintCallable, Category = "Music")
+	virtual void Pause() = 0;
+	UFUNCTION(BlueprintCallable, Category = "Music")
+	virtual void Continue() = 0;
+	UFUNCTION(BlueprintCallable, Category = "Music")
+	virtual void Stop();
+	UFUNCTION(BlueprintCallable, Category = "Music")
+	virtual void Kill();
+
+	operator bool()
+	{
+		return IsValid();
+	}
+
+	UFUNCTION(BlueprintCallable, Category = "Music", Meta = (ExpandEnumAsExecs = "Branches"))
+	virtual void BranchOnTransportState(EMusicHanldeTransportState& Branches) const;
+
+	UFUNCTION(BlueprintCallable, Category = "Music")
+	virtual EMusicHanldeTransportState GetTransportState() const { return CurrentTransportState; }
+
+	UFUNCTION(BlueprintCallable, Category = "Music")
+	virtual bool IsUsingAsset(const TScriptInterface<IMusicalAsset> Asset) const { return IsUsingAsset(Asset.GetObject()); }
+	virtual bool IsUsingAsset(const UObject* Asset) const = 0;
+
+	UFUNCTION(BlueprintCallable, Category = "Music|Clock")
+	virtual TScriptInterface<IMusicEnvironmentClockSource> GetMusicClockSource() = 0;
+
+	UFUNCTION(BlueprintCallable, Category = "Music|Clock", Meta = (ExpandEnumAsExecs = "Branches"))
+	virtual void GetCurrentBarBeat(float& Bar, float& BeatInBar, EMusicHanldeClockValidity& Branches);
+
+	UFUNCTION(BlueprintCallable, Category = "Music|Clock")
+	virtual void BecomeGlobalMusicClockAuthority();
+
+	UFUNCTION(BlueprintCallable, Category = "Music|Clock")
+	virtual void RelinquishGlobalMusicClockAuthority();
+	
+	UFUNCTION(BlueprintCallable, Category = "Music|Clock")
+	virtual void RegisterAsTaggedClock(const FGameplayTag& Tag);
+
+	UFUNCTION(BlueprintCallable, Category = "Music|Clock")
+	virtual void UnregisterAsTaggedClock(const FGameplayTag& Tag = FGameplayTag());
+
+protected:
+	virtual void Stop_Internal() = 0;
+	virtual void Kill_Internal() = 0;
+
+	EMusicHanldeTransportState CurrentTransportState = EMusicHanldeTransportState::PreparingToPlay;
+
+private:
+	bool IsOnGlobalClockAuthorityStack = false;
+	FGameplayTagContainer RegisteredToTags;
+};
+
+using FMusicHandle = TScriptInterface<IMusicHandle>;
+
+struct UE_API FStrongMusicHandle
+{
+public:
+	FORCEINLINE FStrongMusicHandle() = default;
+	FORCEINLINE FStrongMusicHandle(const FStrongMusicHandle& Other) = default;
+	FORCEINLINE FStrongMusicHandle(FStrongMusicHandle&& Other) = default;
+	FORCEINLINE ~FStrongMusicHandle() = default;
+	FORCEINLINE FStrongMusicHandle& operator=(const FStrongMusicHandle& Other) = default;
+	FORCEINLINE FStrongMusicHandle& operator=(FStrongMusicHandle&& Other) = default;
+
+	/**
+	 * A constructor that can take a pointer to any UObject that implements the IMusicHandle interface
+	 */ 
+	template<
+		typename U
+		UE_REQUIRES(std::is_convertible_v<U, TCopyQualifiersFromTo_T<U, UObject>*>)
+	>
+	FStrongMusicHandle(U&& MusicHandleObject)
+	{
+		InterfaceInstance = Cast<IMusicHandle>(ImplicitConv<TCopyQualifiersFromTo_T<U, UObject>*>(MusicHandleObject));
+		if (InterfaceInstance != nullptr)
+		{
+			ObjectInstance = TStrongObjectPtr<UObject>(MusicHandleObject);
+		}
+	}
+
+	FORCEINLINE FStrongMusicHandle& operator=(FMusicHandle Other)
+	{
+		ObjectInstance = TStrongObjectPtr<UObject>(Other.GetObject());
+		InterfaceInstance = Other.GetInterface();
+		return *this;
+	}
+
+	template<
+		typename U
+		UE_REQUIRES(std::is_convertible_v<U, TCopyQualifiersFromTo_T<U, UObject>*>)
+	>
+	FORCEINLINE FStrongMusicHandle& operator=(U& Other)
+	{
+		// Regardless of whether we successfully get an interface from the object,
+		// we need to release any strong pointer we may already have to another object...
+		ObjectInstance = nullptr;
+		InterfaceInstance = Cast<IMusicHandle>(ImplicitConv<TCopyQualifiersFromTo_T<U, UObject>*>(&Other));
+		if (InterfaceInstance != nullptr)
+		{
+			ObjectInstance = TStrongObjectPtr<UObject>(&Other);
+		}
+		return *this;
+	}
+
+	FORCEINLINE FStrongMusicHandle& operator=(std::nullptr_t)
+	{
+		ObjectInstance = nullptr;
+		InterfaceInstance = nullptr;
+		return *this;
+	}
+
+	FORCEINLINE IMusicHandle& operator*() const
+	{
+		check(InterfaceInstance);
+		return *InterfaceInstance;
+	}
+
+	FORCEINLINE IMusicHandle* operator->() const
+	{
+		check(InterfaceInstance);
+		return InterfaceInstance;
+	}
+
+	operator bool() const
+	{
+		return InterfaceInstance != nullptr;
+	}
+
+private:
+	TStrongObjectPtr<UObject> ObjectInstance;
+	IMusicHandle* InterfaceInstance = nullptr;
+};
+
+UCLASS()
+class UMusicHandleBlueprintHelpers : public UBlueprintFunctionLibrary
+{
+	GENERATED_BODY()
+public:
+	UFUNCTION(BlueprintCallable, Category = "Music", Meta = (ExpandEnumAsExecs = "Branches"))
+	static void BranchOnTransportState(TScriptInterface<IMusicHandle> Handle, EMusicHanldeTransportState& Branches);
+
+	UFUNCTION(BlueprintCallable, Category = "Music")
+	static EMusicHanldeTransportState GetTransportState(TScriptInterface<IMusicHandle> Handle);
+};
+
+#undef UE_API // MUSICENVIRONMENT_API
diff --git a/Engine/Plugins/Runtime/MusicEnvironment/Source/MusicEnvironment/Public/MusicTypes/MusicalAsset.h b/Engine/Plugins/Runtime/MusicEnvironment/Source/MusicEnvironment/Public/MusicTypes/MusicalAsset.h
new file mode 100644
--- /dev/null
+++ b/Engine/Plugins/Runtime/MusicEnvironment/Source/MusicEnvironment/Public/MusicTypes/MusicalAsset.h
@@ -0,0 +1,86 @@
+// Copyright Epic Games, Inc. All Rights Reserved.
+
+#pragma once
+
+#include "UObject/Interface.h"
+#include "FrameBasedMusicMap.h"
+#include "GameplayTagContainer.h"
+
+#include "MusicalAsset.generated.h"
+
+#define UE_API MUSICENVIRONMENT_API
+
+/**
+ * This file defines:
+ * ==================
+ * UMusicalAsset: The UInterface derived class needed to support Unreal's interface system. This
+ *                class definition stays empty.
+ *
+ * IMusicalAsset: The actual pure virtual definition of the interface. This is what
+ *                derived classes will inherit from. 
+ * 
+ * FMusicalAsset: This is a convenience type definition that amounts to a TScriptInterface<IMusicalAsset>. 
+ *                ULASSS and USTRUCT things can hold one of these as a member and mark it as a UPROPERTY()
+ *                so the object that implements the script interface will not be garbage collected.
+ *                It looks cleaner.
+ * 
+ * IMusicalAssets are assets that can...
+ *   - Provide a UFrameBasedMusicMap that describes the musical form (tempo and time signature) of the 
+ *     music in the resolution provided by the requestor.
+ *   - Instantiate themselves when they are told to Play or PrepareToPlay. (NOTE: THIS FUNCTIONALITY
+ *     BELONGS IN AN IAudioEmitter class/interface WHEN ONE EXISTS! It is here given the current
+ *     lack of such an emitter system in Unreal Engine. Emitters should be asked to play an asset
+ *     and the caller should get a handle back.)
+ */
+
+class UMusicHandle;
+class UAudioComponent;
+class IMusicHandle;
+
+UINTERFACE(BlueprintType, meta = (CannotImplementInterfaceInBlueprint))
+class UMusicalAsset : public UMusicMapSource
+{
+	GENERATED_BODY()
+};
+
+/**
+ * An asset that has a FrameBasedMusicMap and can be played as music (returning
+ * a 'music handle' that becomes owned by the caller.)
+ */
+class UE_API IMusicalAsset : public IMusicMapSource
+{
+	GENERATED_BODY()
+
+	// Add interface functions to this class. This is the class that will be inherited to implement this interface.
+public:
+	UFUNCTION(BlueprintCallable, Category = "Music", Meta = (DefaultToSelf = "PlaybackContext"))
+	virtual TScriptInterface<IMusicHandle> PrepareToPlay(UObject* PlaybackContext, 
+		UAudioComponent* OnComponent = nullptr,
+		float FromSeconds = 0.0f,
+		bool BecomeAuthoritativeClock = false,
+		const FGameplayTag RegisterAsTaggedClock = FGameplayTag(),
+		bool IsAudition = false);
+
+	UFUNCTION(BlueprintCallable, Category = "Music", Meta = (DefaultToSelf = "PlaybackContext"))
+	virtual TScriptInterface<IMusicHandle> Play(UObject* PlaybackContext,
+		UAudioComponent* OnComponent = nullptr,
+		float FromSeconds = 0.0f,
+		bool BecomeAuthoritativeClock = false,
+		const FGameplayTag RegisterAsTaggedClock = FGameplayTag(),
+		bool IsAudition = false);
+
+protected:
+	virtual TScriptInterface<IMusicHandle> PrepareToPlay_Internal(UObject* PlaybackContext,
+		UAudioComponent* OnComponent = nullptr,
+		float FromSeconds = 0.0f,
+		bool IsAudition = false) = 0;
+
+	virtual TScriptInterface<IMusicHandle> Play_Internal(UObject* PlaybackContext,
+		UAudioComponent* OnComponent = nullptr,
+		float FromSeconds = 0.0f,
+		bool IsAudition = false) = 0;
+};
+
+using FMusicalAsset = TScriptInterface<IMusicalAsset>;
+
+#undef UE_API // MUSICENVIRONMENT_API
