Trivial changes to comments only

diff --git a/Engine/Plugins/Runtime/MusicEnvironment/Source/MusicEnvironment/Public/MusicTypes/MusicHandle.h b/Engine/Plugins/Runtime/MusicEnvironment/Source/MusicEnvironment/Public/MusicTypes/MusicHandle.h
--- a/Engine/Plugins/Runtime/MusicEnvironment/Source/MusicEnvironment/Public/MusicTypes/MusicHandle.h
+++ b/Engine/Plugins/Runtime/MusicEnvironment/Source/MusicEnvironment/Public/MusicTypes/MusicHandle.h
@@ -54,13 +54,24 @@ enum class EMusicHanldeTransportState : uint8
 	Stale,
 };
 
-// This class does not need to be modified.
+ /**
+ * UMusicHandle: The UInterface derived class needed to support Unreal's interface system. This 
+ *               class definition stays empty.
+ */
 UINTERFACE(BlueprintType, meta = (CannotImplementInterfaceInBlueprint))
 class UE_API UMusicHandle : public UInterface
 {
 	GENERATED_BODY()
 };
 
+/**
+ * IMusicHandle: The actual pure virtual definition of the UMusicHandle interface. This is what 
+ *               derived classes will inherit from. AND... IMusicalAssets return 
+ *               TScriptInterface<IMusicHandle> when they are asked to instantiate themselves.
+ *               BUT NOTE... If a non-UCLASS or a non-USTRUCT asks an IMusicalAsset to instantiate 
+ *               itself there may not be anything holding a strong reference to the object the
+ *               script interface is pointing to. More on this below.
+ */
 class UE_API IMusicHandle
 {
 public:
@@ -131,6 +142,13 @@ private:
 
 using FMusicHandle = TScriptInterface<IMusicHandle>;
 
+/** 
+ * FStrongMusicHandle: This wraps a TScriptInterface<IMusicHandle> to hold a TStrongObjectPtr to 
+ *               the UObject that implements the interface. It is like having a TStrongInterfacePtr<>,
+ *               a type that doesn't exist in the engine. This allows non-UCLASS, non-USTRUCT things to
+ *               hold music handles, and for the UStuff inside of those handles to not be garbage collected
+ *               until the FStrongMusicHandle is destroyed.
+ */
 struct UE_API FStrongMusicHandle
 {
 public:
