Trvial changes to comments only.

diff --git a/Engine/Plugins/Runtime/MusicEnvironment/Source/MusicEnvironment/Public/MusicTypes/MusicalAsset.h b/Engine/Plugins/Runtime/MusicEnvironment/Source/MusicEnvironment/Public/MusicTypes/MusicalAsset.h
--- a/Engine/Plugins/Runtime/MusicEnvironment/Source/MusicEnvironment/Public/MusicTypes/MusicalAsset.h
+++ b/Engine/Plugins/Runtime/MusicEnvironment/Source/MusicEnvironment/Public/MusicTypes/MusicalAsset.h
@@ -11,6 +11,14 @@
 #define UE_API MUSICENVIRONMENT_API
 
 /**
+ * IMusicalAssets are assets that can...
+ *   - Provide a UFrameBasedMusicMap that describes the musical form (tempo and time signature) of the
+ *     music in the resolution provided by the requestor.
+ *   - Instantiate themselves when they are told to Play or PrepareToPlay. (NOTE: THIS FUNCTIONALITY
+ *     BELONGS IN AN IAudioEmitter class/interface WHEN ONE EXISTS! It is here given the current
+ *     lack of such an emitter system in Unreal Engine. Emitters should be asked to play an asset
+ *     and the caller should get a handle back.)
+ *
  * This file defines:
  * ==================
  * UMusicalAsset: The UInterface derived class needed to support Unreal's interface system. This
@@ -23,20 +31,16 @@
  *                ULASSS and USTRUCT things can hold one of these as a member and mark it as a UPROPERTY()
  *                so the object that implements the script interface will not be garbage collected.
  *                It looks cleaner.
- * 
- * IMusicalAssets are assets that can...
- *   - Provide a UFrameBasedMusicMap that describes the musical form (tempo and time signature) of the 
- *     music in the resolution provided by the requestor.
- *   - Instantiate themselves when they are told to Play or PrepareToPlay. (NOTE: THIS FUNCTIONALITY
- *     BELONGS IN AN IAudioEmitter class/interface WHEN ONE EXISTS! It is here given the current
- *     lack of such an emitter system in Unreal Engine. Emitters should be asked to play an asset
- *     and the caller should get a handle back.)
  */
 
 class UMusicHandle;
 class UAudioComponent;
 class IMusicHandle;
 
+/**
+ * UMusicalAsset: The UInterface derived class needed to support Unreal's interface system. This
+ *                class definition stays empty.
+ */
 UINTERFACE(BlueprintType, meta = (CannotImplementInterfaceInBlueprint))
 class UMusicalAsset : public UMusicMapSource
 {
@@ -44,8 +48,10 @@ class UMusicalAsset : public UMusicMapSource
 };
 
 /**
- * An asset that has a FrameBasedMusicMap and can be played as music (returning
- * a 'music handle' that becomes owned by the caller.)
+ * IMusicalAsset: The actual pure virtual definition of the interface. This is what
+ *                derived classes will inherit from. It is an asset that has a
+ *                FrameBasedMusicMap and can be played as music (returning a 
+ *                'music handle' that becomes owned by the caller.)
  */
 class UE_API IMusicalAsset : public IMusicMapSource
 {
@@ -81,6 +87,9 @@ protected:
 		bool IsAudition = false) = 0;
 };
 
+/** 
+ * FMusicalAsset: This is a convenience type definition.
+ */
 using FMusicalAsset = TScriptInterface<IMusicalAsset>;
 
 #undef UE_API // MUSICENVIRONMENT_API
