Another trivial comment change

diff --git a/Engine/Plugins/Runtime/MusicEnvironment/Source/MusicEnvironment/Public/MusicTypes/MusicalAsset.h b/Engine/Plugins/Runtime/MusicEnvironment/Source/MusicEnvironment/Public/MusicTypes/MusicalAsset.h
--- a/Engine/Plugins/Runtime/MusicEnvironment/Source/MusicEnvironment/Public/MusicTypes/MusicalAsset.h
+++ b/Engine/Plugins/Runtime/MusicEnvironment/Source/MusicEnvironment/Public/MusicTypes/MusicalAsset.h
@@ -57,7 +57,6 @@ class UE_API IMusicalAsset : public IMusicMapSource
 {
 	GENERATED_BODY()
 
-	// Add interface functions to this class. This is the class that will be inherited to implement this interface.
 public:
 	UFUNCTION(BlueprintCallable, Category = "Music", Meta = (DefaultToSelf = "PlaybackContext"))
 	virtual TScriptInterface<IMusicHandle> PrepareToPlay(UObject* PlaybackContext, 
