Export MusicalAsset class for full external use

diff --git a/Engine/Plugins/Runtime/MusicEnvironment/Source/MusicEnvironment/Public/MusicTypes/MusicalAsset.h b/Engine/Plugins/Runtime/MusicEnvironment/Source/MusicEnvironment/Public/MusicTypes/MusicalAsset.h
--- a/Engine/Plugins/Runtime/MusicEnvironment/Source/MusicEnvironment/Public/MusicTypes/MusicalAsset.h
+++ b/Engine/Plugins/Runtime/MusicEnvironment/Source/MusicEnvironment/Public/MusicTypes/MusicalAsset.h
@@ -42,7 +42,7 @@ class IMusicHandle;
  *                class definition stays empty.
  */
 UINTERFACE(BlueprintType, meta = (CannotImplementInterfaceInBlueprint))
-class UMusicalAsset : public UMusicMapSource
+class UE_API UMusicalAsset : public UMusicMapSource
 {
 	GENERATED_BODY()
 };
