FrameBasedTimeSignature Detail Customization

#rnx

diff --git a/Engine/Plugins/Runtime/MusicEnvironment/MusicEnvironment.uplugin b/Engine/Plugins/Runtime/MusicEnvironment/MusicEnvironment.uplugin
--- a/Engine/Plugins/Runtime/MusicEnvironment/MusicEnvironment.uplugin
+++ b/Engine/Plugins/Runtime/MusicEnvironment/MusicEnvironment.uplugin
@@ -20,6 +20,11 @@
 			"Name": "MusicEnvironment",
 			"Type": "Runtime",
 			"LoadingPhase": "PreDefault"
+		},
+		{
+			"Name": "MusicEnvironmentEditor",
+			"Type": "Editor",
+			"LoadingPhase": "Default"
 		}
 	],
 	"Plugins": [
diff --git a/Engine/Plugins/Runtime/MusicEnvironment/Source/MusicEnvironment/Public/FrameBasedMusicMap.h b/Engine/Plugins/Runtime/MusicEnvironment/Source/MusicEnvironment/Public/FrameBasedMusicMap.h
--- a/Engine/Plugins/Runtime/MusicEnvironment/Source/MusicEnvironment/Public/FrameBasedMusicMap.h
+++ b/Engine/Plugins/Runtime/MusicEnvironment/Source/MusicEnvironment/Public/FrameBasedMusicMap.h
@@ -163,9 +163,9 @@ struct FFrameBasedTimeSignature
 		, Denominator(static_cast<int16>(InDenominator))
 	{}
 
-	UPROPERTY()
+	UPROPERTY(EditAnywhere, Category="Time Signature")
 	int16 Numerator = 4;
-	UPROPERTY()
+	UPROPERTY(EditAnywhere, Category="Time Signature")
 	int16 Denominator = 4;
 
 	bool operator==(const FFrameBasedTimeSignature& Other)
diff --git a/Engine/Plugins/Runtime/MusicEnvironment/Source/MusicEnvironmentEditor/MusicEnvironmentEditor.Build.cs b/Engine/Plugins/Runtime/MusicEnvironment/Source/MusicEnvironmentEditor/MusicEnvironmentEditor.Build.cs
new file mode 100644
--- /dev/null
+++ b/Engine/Plugins/Runtime/MusicEnvironment/Source/MusicEnvironmentEditor/MusicEnvironmentEditor.Build.cs
@@ -0,0 +1,23 @@
+﻿// Copyright Epic Games, Inc. All Rights Reserved.
+
+using UnrealBuildTool;
+
+public class MusicEnvironmentEditor : ModuleRules
+{
+	public MusicEnvironmentEditor(ReadOnlyTargetRules Target) : base(Target)
+	{
+		PCHUsage = ModuleRules.PCHUsageMode.UseExplicitOrSharedPCHs;
+		
+		PrivateDependencyModuleNames.AddRange(
+			new string[]
+			{
+				"Core",
+				"CoreUObject",
+				"MusicEnvironment",
+				"PropertyEditor",
+				"Slate",
+				"SlateCore"
+			}
+		);
+	}
+}
\ No newline at end of file
diff --git a/Engine/Plugins/Runtime/MusicEnvironment/Source/MusicEnvironmentEditor/Private/Customization/FrameBasedTimeSignatureCustomization.cpp b/Engine/Plugins/Runtime/MusicEnvironment/Source/MusicEnvironmentEditor/Private/Customization/FrameBasedTimeSignatureCustomization.cpp
new file mode 100644
--- /dev/null
+++ b/Engine/Plugins/Runtime/MusicEnvironment/Source/MusicEnvironmentEditor/Private/Customization/FrameBasedTimeSignatureCustomization.cpp
@@ -0,0 +1,64 @@
+﻿// Copyright Epic Games, Inc. All Rights Reserved.
+
+#include "FrameBasedTimeSignatureCustomization.h"
+
+#include "DetailWidgetRow.h"
+#include "PropertyHandle.h"
+#include "Widgets/Input/SFrameBasedTimeSignatureInput.h"
+
+void FFrameBasedTimeSignatureCustomization::CustomizeHeader(TSharedRef<IPropertyHandle> PropertyHandle, FDetailWidgetRow& HeaderRow, IPropertyTypeCustomizationUtils& CustomizationUtils)
+{
+	MyPropertyHandle = PropertyHandle.ToSharedPtr();
+	HeaderRow
+	.NameContent()
+	[
+		PropertyHandle->CreatePropertyNameWidget()
+	]
+	.ValueContent()
+	[
+		SNew(SFrameBasedTimeSignatureInput)
+		.Value_Lambda([this]() -> FFrameBasedTimeSignature
+		{
+			if (!MyPropertyHandle.IsValid())
+			{
+				return FFrameBasedTimeSignature();
+			}
+			int16 Numerator;
+			int16 Denominator;
+			TSharedPtr<IPropertyHandle> NumHandle = MyPropertyHandle->GetChildHandle(GET_MEMBER_NAME_CHECKED(FFrameBasedTimeSignature, Numerator));
+			TSharedPtr<IPropertyHandle> DenomHandle = MyPropertyHandle->GetChildHandle(GET_MEMBER_NAME_CHECKED(FFrameBasedTimeSignature, Denominator));
+			check(NumHandle);
+			check(DenomHandle);
+			if (NumHandle->GetValue(Numerator) == FPropertyAccess::Fail)
+			{
+				return FFrameBasedTimeSignature();
+			}
+			
+			if (DenomHandle->GetValue(Denominator) == FPropertyAccess::Fail)
+			{
+				return FFrameBasedTimeSignature();
+			}
+			
+			return FFrameBasedTimeSignature(Numerator, Denominator);
+		})
+		.OnValueCommitted_Lambda([this](const FFrameBasedTimeSignature& NewValue, ETextCommit::Type)
+		{
+			if (!MyPropertyHandle.IsValid())
+			{
+				return;
+			}
+			
+			TSharedPtr<IPropertyHandle> NumHandle = MyPropertyHandle->GetChildHandle(GET_MEMBER_NAME_CHECKED(FFrameBasedTimeSignature, Numerator));
+            TSharedPtr<IPropertyHandle> DenomHandle = MyPropertyHandle->GetChildHandle(GET_MEMBER_NAME_CHECKED(FFrameBasedTimeSignature, Denominator));
+            check(NumHandle);
+            check(DenomHandle);
+			NumHandle->SetValue(NewValue.Numerator);
+			DenomHandle->SetValue(NewValue.Denominator);
+		})
+	];
+}
+
+void FFrameBasedTimeSignatureCustomization::CustomizeChildren(TSharedRef<IPropertyHandle> PropertyHandle, IDetailChildrenBuilder& ChildBuilder, IPropertyTypeCustomizationUtils& CustomizationUtils)
+{
+	
+}
diff --git a/Engine/Plugins/Runtime/MusicEnvironment/Source/MusicEnvironmentEditor/Private/Customization/FrameBasedTimeSignatureCustomization.h b/Engine/Plugins/Runtime/MusicEnvironment/Source/MusicEnvironmentEditor/Private/Customization/FrameBasedTimeSignatureCustomization.h
new file mode 100644
--- /dev/null
+++ b/Engine/Plugins/Runtime/MusicEnvironment/Source/MusicEnvironmentEditor/Private/Customization/FrameBasedTimeSignatureCustomization.h
@@ -0,0 +1,24 @@
+﻿// Copyright Epic Games, Inc. All Rights Reserved.
+
+#pragma once
+
+#include "IPropertyTypeCustomization.h"
+
+struct FFrameBasedTimeSignature;
+
+class FFrameBasedTimeSignatureCustomization : public IPropertyTypeCustomization
+{
+public:
+	static TSharedRef<IPropertyTypeCustomization> MakeInstance()
+	{
+		return MakeShareable(new FFrameBasedTimeSignatureCustomization);
+	}
+
+	virtual void CustomizeHeader(TSharedRef<IPropertyHandle> PropertyHandle, FDetailWidgetRow& HeaderRow, IPropertyTypeCustomizationUtils& CustomizationUtils) override;
+
+	virtual void CustomizeChildren(TSharedRef<IPropertyHandle> PropertyHandle, IDetailChildrenBuilder& ChildBuilder, IPropertyTypeCustomizationUtils& CustomizationUtils) override;
+
+private:
+
+	TSharedPtr<IPropertyHandle> MyPropertyHandle;
+};
\ No newline at end of file
diff --git a/Engine/Plugins/Runtime/MusicEnvironment/Source/MusicEnvironmentEditor/Private/MusicEnvironmentEditorModule.cpp b/Engine/Plugins/Runtime/MusicEnvironment/Source/MusicEnvironmentEditor/Private/MusicEnvironmentEditorModule.cpp
new file mode 100644
--- /dev/null
+++ b/Engine/Plugins/Runtime/MusicEnvironment/Source/MusicEnvironmentEditor/Private/MusicEnvironmentEditorModule.cpp
@@ -0,0 +1,27 @@
+﻿// Copyright Epic Games, Inc. All Rights Reserved.
+
+#include "MusicEnvironmentEditorModule.h"
+#include "Modules/ModuleInterface.h"
+#include "PropertyEditorModule.h"
+#include "Customization/FrameBasedTimeSignatureCustomization.h"
+#include "Modules/ModuleManager.h"
+
+class FMusicEnvironmentEditorModule : public IModuleInterface
+{
+public:
+ 
+	virtual void StartupModule() override
+	{
+		// register detail customizations
+		FPropertyEditorModule& PropertyEditorModule = FModuleManager::LoadModuleChecked<FPropertyEditorModule>("PropertyEditor");
+		PropertyEditorModule.RegisterCustomPropertyTypeLayout("FrameBasedTimeSignature", FOnGetPropertyTypeCustomizationInstance::CreateStatic(&FFrameBasedTimeSignatureCustomization::MakeInstance));
+	}
+ 
+	virtual void ShutdownModule() override
+	{
+		FPropertyEditorModule& PropertyEditorModule = FModuleManager::LoadModuleChecked<FPropertyEditorModule>("PropertyEditor");
+		PropertyEditorModule.UnregisterCustomPropertyTypeLayout("FrameBasedTimeSignature");
+	}
+};
+
+IMPLEMENT_MODULE(FMusicEnvironmentEditorModule, MusicEnvironmentEditor);
\ No newline at end of file
diff --git a/Engine/Plugins/Runtime/MusicEnvironment/Source/MusicEnvironmentEditor/Private/Widgets/Input/SFrameBasedTimeSignatureInput.cpp b/Engine/Plugins/Runtime/MusicEnvironment/Source/MusicEnvironmentEditor/Private/Widgets/Input/SFrameBasedTimeSignatureInput.cpp
new file mode 100644
--- /dev/null
+++ b/Engine/Plugins/Runtime/MusicEnvironment/Source/MusicEnvironmentEditor/Private/Widgets/Input/SFrameBasedTimeSignatureInput.cpp
@@ -0,0 +1,108 @@
+﻿// Copyright Epic Games, Inc. All Rights Reserved.
+
+
+#include "Widgets/Input/SFrameBasedTimeSignatureInput.h"
+
+#include "FrameBasedMusicMap.h"
+#include "Internationalization/Regex.h"
+#include "Widgets/Input/SEditableTextBox.h"
+
+#define LOCTEXT_NAMESPACE "SFrameBasedTimeSignatureInput"
+
+void SFrameBasedTimeSignatureInput::Construct(const FArguments& InArgs)
+{
+	ValueAttribute = InArgs._Value;
+	MaxNumerator = InArgs._MaxNumerator;
+	MaxDenominator = InArgs._MaxDenominator;
+	OnValueCommitted = InArgs._OnValueCommitted;
+	ChildSlot
+	[
+		SNew(SEditableTextBox)
+		.Font(InArgs._Font)
+		.SelectAllTextWhenFocused(true)
+		.ClearKeyboardFocusOnCommit(false)
+		.Text(this, &SFrameBasedTimeSignatureInput::GetValueAsText)
+		.RevertTextOnEscape(true)
+		.OnTextCommitted(this, &SFrameBasedTimeSignatureInput::TextBox_OnTextCommitted)
+		.OnVerifyTextChanged(this, &SFrameBasedTimeSignatureInput::TextBox_VerifyTextChanged)
+		.IsEnabled(InArgs._IsEnabled)
+	];
+}
+
+FFrameBasedTimeSignature SFrameBasedTimeSignatureInput::GetValue() const
+{
+	return ValueAttribute.Get();
+}
+
+void SFrameBasedTimeSignatureInput::SetValue(const FFrameBasedTimeSignature& Value)
+{
+	ValueAttribute.Set(Value);
+	OnValueCommitted.ExecuteIfBound(Value, ETextCommit::Default);
+}
+
+
+FString SFrameBasedTimeSignatureInput::GetValueAsString() const
+{
+	FFrameBasedTimeSignature CurrentValue = ValueAttribute.Get();
+	return FString::Printf(TEXT("%d/%d"), CurrentValue.Numerator, CurrentValue.Denominator);
+}
+
+FText SFrameBasedTimeSignatureInput::GetValueAsText() const
+{
+	return FText::FromString(GetValueAsString());
+}
+
+TOptional<FFrameBasedTimeSignature> SFrameBasedTimeSignatureInput::ParseValueFromString(const FString& InString)
+{
+	TOptional<FFrameBasedTimeSignature> OutTimeSig;
+	FRegexPattern TimeSigPattern(TEXT("^(\\d+)\\/(\\d+)$"));
+	FRegexMatcher TimeSigMatcher(TimeSigPattern, InString.Replace(TEXT(" "), TEXT("")));
+	if (TimeSigMatcher.FindNext())
+	{
+		OutTimeSig = FFrameBasedTimeSignature(
+			FCString::Atoi(*TimeSigMatcher.GetCaptureGroup(1)),
+			FCString::Atoi(*TimeSigMatcher.GetCaptureGroup(2)));
+	}
+
+	return OutTimeSig;
+}
+
+void SFrameBasedTimeSignatureInput::TextBox_OnTextCommitted(const FText& NewText, ETextCommit::Type CommitInfo)
+{
+	TOptional<FFrameBasedTimeSignature> NewValue = ParseValueFromString(NewText.ToString());
+	if (NewValue.IsSet())
+	{
+		if (!ValueAttribute.IsBound())
+		{
+			ValueAttribute.Set(NewValue.GetValue());
+		}
+		
+		OnValueCommitted.ExecuteIfBound(NewValue.GetValue(), CommitInfo);
+	}
+}
+
+bool SFrameBasedTimeSignatureInput::TextBox_VerifyTextChanged(const FText& NewText, FText& OutErrorMessage) const
+{
+	TOptional<FFrameBasedTimeSignature> NewValue = ParseValueFromString(NewText.ToString());
+	if (NewValue.IsSet())
+	{
+		int32 MaxNum = FMath::Max(1, MaxNumerator.IsSet() ? MaxNumerator.Get() : MAX_int32);
+		if (NewValue->Numerator <= 0 || NewValue->Numerator > MaxNum)
+		{
+			OutErrorMessage = FText::Format(LOCTEXT("InputError_Numerator", "Time signature numerator must be from 1 to {0}"), MaxNum);
+			return false;
+		}
+
+		int32 MaxDen = FMath::Max(1, MaxDenominator.IsSet() ? MaxDenominator.Get() : MAX_int32);
+		if (NewValue->Denominator <= 0 || NewValue->Denominator > MaxDen)
+		{
+			OutErrorMessage = FText::Format(LOCTEXT("InputError_Denominator", "Time signature denominator must be from 1 to {0}"), MaxNum);
+			return false;
+		}
+		return true;
+	}
+	OutErrorMessage = LOCTEXT("InputError_TimeSignature", "Time signature must be in the form '4/4', '7/8', etc.");
+	return false;
+}
+
+#undef LOCTEXT_NAMESPACE //SFrameBasedTimeSignatureInput
\ No newline at end of file
diff --git a/Engine/Plugins/Runtime/MusicEnvironment/Source/MusicEnvironmentEditor/Public/MusicEnvironmentEditorModule.h b/Engine/Plugins/Runtime/MusicEnvironment/Source/MusicEnvironmentEditor/Public/MusicEnvironmentEditorModule.h
new file mode 100644
--- /dev/null
+++ b/Engine/Plugins/Runtime/MusicEnvironment/Source/MusicEnvironmentEditor/Public/MusicEnvironmentEditorModule.h
@@ -0,0 +1,3 @@
+﻿// Copyright Epic Games, Inc. All Rights Reserved.
+
+#pragma once
\ No newline at end of file
diff --git a/Engine/Plugins/Runtime/MusicEnvironment/Source/MusicEnvironmentEditor/Public/Widgets/Input/SFrameBasedTimeSignatureInput.h b/Engine/Plugins/Runtime/MusicEnvironment/Source/MusicEnvironmentEditor/Public/Widgets/Input/SFrameBasedTimeSignatureInput.h
new file mode 100644
--- /dev/null
+++ b/Engine/Plugins/Runtime/MusicEnvironment/Source/MusicEnvironmentEditor/Public/Widgets/Input/SFrameBasedTimeSignatureInput.h
@@ -0,0 +1,66 @@
+﻿// Copyright Epic Games, Inc. All Rights Reserved.
+
+#pragma once
+
+#include "Widgets/SCompoundWidget.h"
+#include "FrameBasedMusicMap.h"
+#include "Styling/CoreStyle.h"
+
+class MUSICENVIRONMENTEDITOR_API SFrameBasedTimeSignatureInput : public SCompoundWidget
+{
+public:
+
+	/** Notification for numeric value committed */
+	DECLARE_DELEGATE_TwoParams(FOnValueCommitted, const FFrameBasedTimeSignature&, ETextCommit::Type);
+
+	SLATE_BEGIN_ARGS(SFrameBasedTimeSignatureInput)
+		: _Value(FFrameBasedTimeSignature())
+		, _MaxNumerator(DefaultMaxNumerator)
+		, _MaxDenominator(DefaultMaxDenominator)
+		, _IsEnabled(true)
+		, _OnValueCommitted()
+		, _Font(FCoreStyle::Get().GetFontStyle(TEXT("NormalFont")))
+	{}
+
+	SLATE_ATTRIBUTE(FFrameBasedTimeSignature, Value)
+		
+	SLATE_ATTRIBUTE(int32, MaxNumerator)
+
+	SLATE_ATTRIBUTE(int32, MaxDenominator)
+
+	SLATE_ATTRIBUTE(bool, IsEnabled)
+
+	SLATE_EVENT(FOnValueCommitted, OnValueCommitted)
+		
+	SLATE_ATTRIBUTE(FSlateFontInfo, Font)
+		
+SLATE_END_ARGS()
+
+void Construct(const FArguments& InArgs);
+
+	FFrameBasedTimeSignature GetValue() const;
+
+	void SetValue(const FFrameBasedTimeSignature& Value);
+	
+private:
+
+	/** The default minimum fractional digits */
+	static constexpr int32 DefaultMaxNumerator = 99;
+
+	/** The default maximum fractional digits */
+	static constexpr int32 DefaultMaxDenominator = 99;
+
+	TAttribute<FFrameBasedTimeSignature> ValueAttribute;
+	TAttribute< int32 > MaxNumerator;
+	TAttribute< int32 > MaxDenominator;
+	FOnValueCommitted OnValueCommitted;
+
+	FString GetValueAsString() const;
+	FText GetValueAsText() const;
+
+	static TOptional<FFrameBasedTimeSignature> ParseValueFromString(const FString& InString); 
+
+	void TextBox_OnTextCommitted(const FText& NewText, ETextCommit::Type CommitInfo);
+	
+	bool TextBox_VerifyTextChanged(const FText& InText, FText& OutErrorMessage) const; 
+};
\ No newline at end of file
