Warn if the Musical Asset's music map has changed

#rnx

diff --git a/Engine/Plugins/Runtime/MusicEnvironment/Source/MusicEnvironment/Public/FrameBasedMusicMap.h b/Engine/Plugins/Runtime/MusicEnvironment/Source/MusicEnvironment/Public/FrameBasedMusicMap.h
--- a/Engine/Plugins/Runtime/MusicEnvironment/Source/MusicEnvironment/Public/FrameBasedMusicMap.h
+++ b/Engine/Plugins/Runtime/MusicEnvironment/Source/MusicEnvironment/Public/FrameBasedMusicMap.h
@@ -104,6 +104,13 @@ struct FFrameBasedTempoPoint
 		MicrosecondsPerQuarterNote = InBpm >= 0.0f ? static_cast<int32>(60000000.0f / InBpm) : 0;
 	}
 
+	bool IsNearlyEqual(const FFrameBasedTempoPoint& Other, float MsTolerance = UE_SMALL_NUMBER) const
+	{
+		return OnTick == Other.OnTick
+			&& MicrosecondsPerQuarterNote == Other.MicrosecondsPerQuarterNote
+			&& FMath::IsNearlyEqual(OnMs, Other.OnMs, MsTolerance);
+	}
+
 	UPROPERTY()
 	int32 MicrosecondsPerQuarterNote = 500000; // 120 bpm default.
 	UPROPERTY()
@@ -149,6 +156,24 @@ struct FFrameBasedTempoPoint
 		float DeltaTick = static_cast<float>(Tick - OnTick);
 		return OnMs + (DeltaTick * MsPerTick);
 	}
+
+	static bool Equals(const TArray<FFrameBasedTempoPoint>& A, const TArray<FFrameBasedTempoPoint>& B)
+	{
+		if (A.Num() != B.Num())
+		{
+			return false;
+		}
+		
+		for (int32 Idx = 0; Idx < A.Num(); ++Idx)
+		{
+			if (!A[Idx].IsNearlyEqual(B[Idx]))
+			{
+				return false;
+			}
+		}
+
+		return true;
+	}
 };
 
 USTRUCT()
@@ -168,12 +193,12 @@ struct FFrameBasedTimeSignature
 	UPROPERTY(EditAnywhere, Category="Time Signature")
 	int16 Denominator = 4;
 
-	bool operator==(const FFrameBasedTimeSignature& Other)
+	bool operator==(const FFrameBasedTimeSignature& Other) const
 	{
 		return Numerator == Other.Numerator && Denominator == Other.Denominator;
 	}
 
-	bool operator!=(const FFrameBasedTimeSignature& Other)
+	bool operator!=(const FFrameBasedTimeSignature& Other) const
 	{
 		return Numerator != Other.Numerator || Denominator != Other.Denominator;
 	}
@@ -195,6 +220,14 @@ struct FFrameBasedTimeSignaturePoint
 	{
 	}
 
+	bool IsNearlyEqual(const FFrameBasedTimeSignaturePoint& Other, double FrameTolerance = UE_DOUBLE_SMALL_NUMBER) const
+	{
+		return OnTick == Other.OnTick
+			&& OnBar == Other.OnBar
+			&& TimeSignature == Other.TimeSignature
+			&& FMath::IsNearlyEqual(OnFrame, Other.OnFrame, FrameTolerance);
+	}
+
 	UPROPERTY()
 	FFrameBasedTimeSignature TimeSignature;
 	UPROPERTY()
@@ -255,6 +288,24 @@ struct FFrameBasedTimeSignaturePoint
 		float TicksIntoBar = Beat * static_cast<float>(TicksPerBeat());
 		return OnTick + (BarDelta * TicksPerBar()) + FMath::RoundToInt32(TicksIntoBar);
 	}
+
+	static bool Equals(const TArray<FFrameBasedTimeSignaturePoint>& A, const TArray<FFrameBasedTimeSignaturePoint>& B)
+	{
+		if (A.Num() != B.Num())
+		{
+			return false;
+		}
+		
+		for (int32 Idx = 0; Idx < A.Num(); ++Idx)
+	    {
+	        if (A[Idx].IsNearlyEqual(B[Idx]))
+	        {
+        		return false;
+	        }
+	    }
+
+		return true;
+	}
 };
 
 USTRUCT()
