Move IMusicMapSource to its own file

#rnx

diff --git a/Engine/Plugins/Runtime/Harmonix/Source/HarmonixMidi/Private/HarmonixMidi/MidiFile.cpp b/Engine/Plugins/Runtime/Harmonix/Source/HarmonixMidi/Private/HarmonixMidi/MidiFile.cpp
--- a/Engine/Plugins/Runtime/Harmonix/Source/HarmonixMidi/Private/HarmonixMidi/MidiFile.cpp
+++ b/Engine/Plugins/Runtime/Harmonix/Source/HarmonixMidi/Private/HarmonixMidi/MidiFile.cpp
@@ -9,6 +9,7 @@
 #include "HarmonixMidi/BarMap.h"
 #include "HarmonixMidi/SongMaps.h"
 #include "HarmonixMidi/SongMapReceiver.h"
+#include "FrameBasedMusicMap.h"
 #include "Misc/Paths.h"
 #include "Algo/ForEach.h"
 #include "UObject/AssetRegistryTagsContext.h"
diff --git a/Engine/Plugins/Runtime/Harmonix/Source/HarmonixMidi/Public/HarmonixMidi/MidiFile.h b/Engine/Plugins/Runtime/Harmonix/Source/HarmonixMidi/Public/HarmonixMidi/MidiFile.h
--- a/Engine/Plugins/Runtime/Harmonix/Source/HarmonixMidi/Public/HarmonixMidi/MidiFile.h
+++ b/Engine/Plugins/Runtime/Harmonix/Source/HarmonixMidi/Public/HarmonixMidi/MidiFile.h
@@ -2,14 +2,13 @@
 
 #pragma once
 
-#include "FrameBasedMusicMap.h"
 #include "HarmonixMidi/TempoMap.h"
 #include "HarmonixMidi/MidiMsg.h"
 #include "HarmonixMidi/BarMap.h"
 #include "HarmonixMidi/SongMaps.h"
 #include "HarmonixMidi/MidiTrack.h"
 #include "IAudioProxyInitializer.h"
-
+#include "MusicTypes/MusicMapSource.h"
 #include "MidiFile.generated.h"
 
 #define UE_API HARMONIXMIDI_API
diff --git a/Engine/Plugins/Runtime/MusicEnvironment/Source/MusicEnvironment/Public/FrameBasedMusicMap.h b/Engine/Plugins/Runtime/MusicEnvironment/Source/MusicEnvironment/Public/FrameBasedMusicMap.h
--- a/Engine/Plugins/Runtime/MusicEnvironment/Source/MusicEnvironment/Public/FrameBasedMusicMap.h
+++ b/Engine/Plugins/Runtime/MusicEnvironment/Source/MusicEnvironment/Public/FrameBasedMusicMap.h
@@ -464,55 +464,4 @@ private:
 	UE_API FString MakeBarBeatString(double DecimalFrame, const FFrameBasedTimeSignaturePoint* TimeSignature, const FFrameBasedTempoPoint* Tempo, FrameBasedMusicMap::EBarBeatFormat BarBeatFormat) const;
 };
 
-UINTERFACE(MinimalAPI)
-class UMusicMapSource : public UInterface
-{
-	GENERATED_BODY()
-};
-
-struct FMarkerProviderEntry
-{
-	FString Label;
-	FFrameNumber FrameNumber;
-
-	FMarkerProviderEntry(const FString& InLabel, FFrameNumber InFrameNumber)
-		: Label(InLabel)
-		, FrameNumber(InFrameNumber)
-	{}
-};
-
-USTRUCT()
-struct FMarkerProviderResults
-{
-	GENERATED_BODY()
-
-	struct FChannel
-	{
-		FString Name;
-		TArray<FMarkerProviderEntry> Markers;
-
-		FChannel(const FString& InName)
-			: Name(InName)
-		{}
-	};
-
-	bool IsEmpty() const { return Channels.IsEmpty(); }
-
-	TArray<FChannel> Channels;
-};
-
-class IMusicMapSource
-{
-public:
-
-	GENERATED_BODY()
-
-	UFUNCTION()
-	virtual void CreateFrameBasedMusicMap(UFrameBasedMusicMap* Map) const { return; }
-
-	UFUNCTION()
-	virtual FMarkerProviderResults GatherMarkers(const UFrameBasedMusicMap* Map) const { return FMarkerProviderResults(); }
-
-};
-
 #undef UE_API
diff --git a/Engine/Plugins/Runtime/MusicEnvironment/Source/MusicEnvironment/Public/MusicTypes/MusicMapSource.h b/Engine/Plugins/Runtime/MusicEnvironment/Source/MusicEnvironment/Public/MusicTypes/MusicMapSource.h
new file mode 100644
--- /dev/null
+++ b/Engine/Plugins/Runtime/MusicEnvironment/Source/MusicEnvironment/Public/MusicTypes/MusicMapSource.h
@@ -0,0 +1,63 @@
+﻿// Copyright Epic Games, Inc. All Rights Reserved.
+
+#pragma once
+
+#include "CoreTypes.h"
+#include "UObject/Interface.h"
+#include "MusicMapSource.generated.h"
+
+class UFrameBasedMusicMap;
+
+#define UE_API MUSICENVIRONMENT_API
+
+struct FMarkerProviderEntry
+{
+	FString Label;
+	FFrameNumber FrameNumber;
+
+	FMarkerProviderEntry(const FString& InLabel, FFrameNumber InFrameNumber)
+		: Label(InLabel)
+		, FrameNumber(InFrameNumber)
+	{}
+};
+
+USTRUCT()
+struct FMarkerProviderResults
+{
+	GENERATED_BODY()
+
+	struct FChannel
+	{
+		FString Name;
+		TArray<FMarkerProviderEntry> Markers;
+
+		FChannel(const FString& InName)
+			: Name(InName)
+		{}
+	};
+
+	bool IsEmpty() const { return Channels.IsEmpty(); }
+
+	TArray<FChannel> Channels;
+};
+
+UINTERFACE(MinimalAPI)
+class UMusicMapSource : public UInterface
+{
+	GENERATED_BODY()
+};
+
+class IMusicMapSource
+{
+public:
+
+	GENERATED_BODY()
+
+	UFUNCTION()
+	virtual void CreateFrameBasedMusicMap(UFrameBasedMusicMap* Map) const { return; }
+
+	UFUNCTION()
+	virtual FMarkerProviderResults GatherMarkers(const UFrameBasedMusicMap* Map) const { return FMarkerProviderResults(); }
+};
+
+#undef UE_API // MUSICENVIRONMENT_API
\ No newline at end of file
diff --git a/Engine/Plugins/Runtime/MusicEnvironment/Source/MusicEnvironment/Public/MusicTypes/MusicalAsset.h b/Engine/Plugins/Runtime/MusicEnvironment/Source/MusicEnvironment/Public/MusicTypes/MusicalAsset.h
--- a/Engine/Plugins/Runtime/MusicEnvironment/Source/MusicEnvironment/Public/MusicTypes/MusicalAsset.h
+++ b/Engine/Plugins/Runtime/MusicEnvironment/Source/MusicEnvironment/Public/MusicTypes/MusicalAsset.h
@@ -3,9 +3,8 @@
 #pragma once
 
 #include "UObject/Interface.h"
-#include "FrameBasedMusicMap.h"
 #include "GameplayTagContainer.h"
-
+#include "MusicTypes/MusicMapSource.h"
 #include "MusicalAsset.generated.h"
 
 #define UE_API MUSICENVIRONMENT_API
