Implement Insert Tempo & Time Signature changes on FrameBasedMusicMap

#rnx

diff --git a/Engine/Plugins/Runtime/MusicEnvironment/Source/MusicEnvironment/Private/FrameBasedMusicMap.cpp b/Engine/Plugins/Runtime/MusicEnvironment/Source/MusicEnvironment/Private/FrameBasedMusicMap.cpp
--- a/Engine/Plugins/Runtime/MusicEnvironment/Source/MusicEnvironment/Private/FrameBasedMusicMap.cpp
+++ b/Engine/Plugins/Runtime/MusicEnvironment/Source/MusicEnvironment/Private/FrameBasedMusicMap.cpp
@@ -572,6 +572,275 @@ void UFrameBasedMusicMap::AddTempo(int32 AtTick, float AtMs, float Bpm)
 	}
 }
 
+void UFrameBasedMusicMap::InsertTimeSignature(int32 AtBar, int16 Numerator, int16 Denominator)
+{
+	FFrameBasedTimeSignature InTimeSig(Numerator, Denominator);
+	if (!ensure(!BarMap.IsEmpty()))
+	{
+		return;
+	}
+	
+	int32 Index = BarMap.FindLastByPredicate([AtBar](FFrameBasedTimeSignaturePoint const& Point) { return AtBar >= Point.OnBar; });
+	check(Index != INDEX_NONE);
+	FFrameBasedTimeSignaturePoint& Point = BarMap[Index];
+	if (Point.TimeSignature == InTimeSig)
+	{
+		// time signature is already at the desired time signature
+		return;
+	}
+	if (Point.OnBar == AtBar)
+	{
+		// update the time signature of the point at the given bar
+		Point.TimeSignature = InTimeSig;
+		int32 PrevIndex = Index - 1;
+		if (BarMap.IsValidIndex(PrevIndex) && BarMap[PrevIndex].TimeSignature == InTimeSig)
+		{
+			BarMap.RemoveAt(Index);
+			Index = PrevIndex;
+		}
+	}
+	else
+	{
+		int32 AtTick =  BarMap[Index].BarToTick(AtBar);
+		// insert the new time signature at the given tick and index.
+		// Note: The index gets incremented here in order to insert the new point after the one we found
+		Index = BarMap.Insert(FFrameBasedTimeSignaturePoint(InTimeSig, AtBar, AtTick, GetFrame(AtBar)), Index + 1);
+	}
+
+	// update the successive time signatures to their new ticks based on the time signature change
+	for (int32 Idx = Index + 1; Idx < BarMap.Num(); /* Idx is conditionally iterated in body */)
+	{
+		// only remove the immediately "next" time signature if it's a duplicate
+		if (Idx == Index + 1 && BarMap[Idx].TimeSignature == InTimeSig)
+		{
+			BarMap.RemoveAt(Idx);
+			continue;
+		}
+		
+		int32 NewTick = BarMap[Idx - 1].BarToTick(BarMap[Idx].OnBar);
+		const FFrameBasedTempoPoint* TempoAtFrame = GetTempoPointForTick(NewTick);
+		BarMap[Idx].OnTick = NewTick;
+		BarMap[Idx].OnFrame = FrameResolution.AsFrameTime(TempoAtFrame->TickToSeconds(NewTick)).AsDecimal();
+		++Idx;
+	}
+}
+
+void UFrameBasedMusicMap::InsertTempo(int32 AtTick, float InBpm)
+{
+	if (!ensure(!TempoMap.IsEmpty()))
+	{
+		return;
+	}
+	
+	int32 TempoIndex = Algo::UpperBoundBy(TempoMap, AtTick, &FFrameBasedTempoPoint::OnTick) - 1;
+	check(TempoIndex != INDEX_NONE);
+	FFrameBasedTempoPoint& Point = TempoMap[TempoIndex];
+
+	if (FMath::IsNearlyEqual(Point.Bpm(), InBpm, 0.001f))
+	{
+		// previous tempo point is at the desired tempo, so no need to add tempo change
+		return;
+	}
+	
+	if (Point.OnTick == AtTick)
+	{
+		Point.SetBpm(InBpm);
+		int32 PrevIndex = TempoIndex - 1;
+		if (TempoMap.IsValidIndex(PrevIndex) && FMath::IsNearlyEqual(TempoMap[PrevIndex].Bpm() , InBpm, 0.001f))
+		{
+			TempoMap.RemoveAt(TempoIndex);
+			TempoIndex = PrevIndex;
+		}
+	}
+	else
+	{
+		float ExpectedMs = TempoMap[TempoIndex].TickToMs(AtTick);
+		TempoIndex = TempoMap.Insert(FFrameBasedTempoPoint(InBpm, AtTick, ExpectedMs), TempoIndex + 1);
+	}
+
+	for (int32 Idx = TempoIndex + 1; Idx < TempoMap.Num(); /* Duplicate tempos are removed. Idx is conditionally iterated in body */ )
+	{
+		if (Idx == TempoIndex + 1 && FMath::IsNearlyEqual(TempoMap[Idx].Bpm(), InBpm))
+		{
+			UE_LOG(LogMusicMap, Log, TEXT("Removing duplicate tempo change"));
+			TempoMap.RemoveAt(Idx);
+			continue;
+		}
+		
+		TempoMap[Idx].OnMs = TempoMap[Idx - 1].TickToMs(TempoMap[Idx].OnTick);
+		++Idx;
+	}
+
+	// remap the bar tick frame times based on the new tempo change
+	int32 TimeSigIndex = BarMap.IndexOfByPredicate([AtTick](const FFrameBasedTimeSignaturePoint& Point) { return AtTick >= Point.OnTick; });
+	for (int32 Idx = TimeSigIndex + 1; Idx < BarMap.Num(); ++Idx)
+	{
+		int32 BarTick = BarMap[Idx].OnTick;
+		const FFrameBasedTempoPoint* TempoAtFrame = GetTempoPointForTick(BarTick);
+		BarMap[Idx].OnFrame = FrameResolution.AsFrameTime(TempoAtFrame->TickToSeconds(BarTick)).AsDecimal();
+	}
+}
+
+void UFrameBasedMusicMap::InsertTempo(FMusicalTime AtMusicalTime, float Bpm)
+{
+	if (const FFrameBasedTimeSignaturePoint* TimeSigPoint = GetTimeSignaturePointForBar(AtMusicalTime.Bar))
+	{
+		int32 Tick = TimeSigPoint->MusicalTimeToTick(AtMusicalTime);
+		InsertTempo(Tick, Bpm);
+	}
+}
+
+void UFrameBasedMusicMap::InsertTimeSignature(FMusicalTime AtMusicalTime, const FFrameBasedTimeSignature& TimeSignature)
+{
+	InsertTimeSignature(AtMusicalTime.Bar, TimeSignature.Numerator, TimeSignature.Denominator);
+}
+
+bool UFrameBasedMusicMap::RemoveTimeSignature(int32 AtBar)
+{
+	int32 Index = Algo::BinarySearchBy(BarMap, AtBar, &FFrameBasedTimeSignaturePoint::OnBar);
+	
+	if (Index == INDEX_NONE)
+	{
+		return false;
+	}
+
+	// don't remove the first and only time signature change
+	if (Index == 0)
+	{
+		if (BarMap.Num() == 1)
+		{
+			return false;
+		}
+		// instead of removing the first time signature
+		// change it to the successive time signature
+		// and remove that one instead
+		BarMap[0].TimeSignature = BarMap[1].TimeSignature;
+		Index = Index + 1;
+	}
+	
+	BarMap.RemoveAt(Index);
+
+	int32 PrevIndex = Index - 1;
+	if (BarMap.IsValidIndex(Index) && BarMap.IsValidIndex(PrevIndex))
+	{
+		if (BarMap[PrevIndex].TimeSignature == BarMap[Index].TimeSignature)
+		{
+			BarMap.RemoveAt(Index);
+		}
+	}
+
+	// update the successive time signatures to their new ticks based on the time signature change
+	for (int32 Idx = Index; Idx < BarMap.Num(); ++Idx)
+	{
+		int32 NewTick = BarMap[Idx - 1].BarToTick(BarMap[Idx].OnBar);
+		const FFrameBasedTempoPoint* TempoAtFrame = GetTempoPointForTick(NewTick);
+		BarMap[Idx].OnTick = NewTick;
+		BarMap[Idx].OnFrame = FrameResolution.AsFrameTime(TempoAtFrame->TickToSeconds(NewTick)).AsDecimal();
+	}
+	
+	return true;
+}
+
+bool UFrameBasedMusicMap::RemoveTempo(int32 AtTick)
+{
+	int32 Index = Algo::BinarySearchBy(TempoMap, AtTick, &FFrameBasedTempoPoint::OnTick);
+
+	if (Index == INDEX_NONE)
+	{
+		return false;
+	}
+
+	if (Index == 0)
+	{
+		if (TempoMap.Num() == 1)
+		{
+			return false;
+		}
+		
+		// when removing the first tempo change,
+		// 
+		TempoMap[0].MicrosecondsPerQuarterNote = TempoMap[1].MicrosecondsPerQuarterNote;
+		Index = Index + 1;
+	}
+
+	TempoMap.RemoveAt(Index);
+
+	int32 PrevIndex = Index - 1;
+	if (TempoMap.IsValidIndex(Index) && TempoMap.IsValidIndex(PrevIndex))
+	{
+		if (TempoMap[PrevIndex].MicrosecondsPerQuarterNote == TempoMap[Index].MicrosecondsPerQuarterNote)
+		{
+			TempoMap.RemoveAt(Index);
+		}
+	}
+	
+	for (int32 Idx = Index; Idx < TempoMap.Num(); ++Idx)
+	{
+		TempoMap[Idx].OnMs = TempoMap[Idx - 1].TickToMs(TempoMap[Idx].OnTick);
+	}
+
+	// remap the bar tick frame times based on the new tempo change
+	int32 TimeSigIndex = BarMap.IndexOfByPredicate([AtTick](const FFrameBasedTimeSignaturePoint& Point) { return AtTick >= Point.OnTick; });
+	for (int32 Idx = TimeSigIndex + 1; Idx < BarMap.Num(); ++Idx)
+	{
+		int32 BarTick = BarMap[Idx].OnTick;
+		const FFrameBasedTempoPoint* TempoAtFrame = GetTempoPointForTick(BarTick);
+		BarMap[Idx].OnFrame = FrameResolution.AsFrameTime(TempoAtFrame->TickToSeconds(BarTick)).AsDecimal();
+	}
+	
+	return true;
+}
+
+bool UFrameBasedMusicMap::RemoveTimeSignature(FMusicalTime AtMusicalTime)
+{
+	if (AtMusicalTime.TickInBar != 0)
+	{
+		return false;
+	}
+
+	return RemoveTimeSignature(AtMusicalTime.Bar);
+}
+
+bool UFrameBasedMusicMap::RemoveTempo(FMusicalTime AtMusicalTime)
+{
+	if (const FFrameBasedTimeSignaturePoint* TimeSigPoint = GetTimeSignaturePointForBar(AtMusicalTime.Bar))
+	{
+		int32 AtTick = TimeSigPoint->MusicalTimeToTick(AtMusicalTime);
+		RemoveTempo(AtTick);
+	}
+	return false;
+}
+
+
+bool UFrameBasedMusicMap::HasTimeSignatureChange(int32 AtBar)
+{
+	return INDEX_NONE != Algo::BinarySearchBy(BarMap, AtBar, &FFrameBasedTimeSignaturePoint::OnBar);
+}
+
+bool UFrameBasedMusicMap::HasTempoChange(int32 AtTick)
+{
+	return INDEX_NONE != Algo::BinarySearchBy(TempoMap, AtTick, &FFrameBasedTempoPoint::OnTick);
+}
+
+bool UFrameBasedMusicMap::HasTimeSignatureChange(FMusicalTime AtMusicalTime)
+{
+	if (AtMusicalTime.TickInBar != 0)
+	{
+		return false;
+	}
+	return HasTimeSignatureChange(AtMusicalTime.Bar);
+}
+
+bool UFrameBasedMusicMap::HasTempoChange(FMusicalTime AtMusicalTime)
+{
+	if (const FFrameBasedTimeSignaturePoint* TimeSigPoint = GetTimeSignaturePointForBar(AtMusicalTime.Bar))
+	{
+		int32 AtTick = TimeSigPoint->MusicalTimeToTick(AtMusicalTime);
+		return HasTempoChange(AtTick);
+	}
+	return false;
+}
+
 void UFrameBasedMusicMap::GatherSubdivisonsForTimespan(int32 StartTick, int32 EndTick, int32 QuantizationInterval, TArray<FGridInterval>& OutSubdivisions ) const
 {
 	FMusicalTime StartMusicalTime = Quantize(GetMusicalTime(StartTick), QuantizationInterval, EQuantizeDirection::Ceiling);
diff --git a/Engine/Plugins/Runtime/MusicEnvironment/Source/MusicEnvironment/Public/FrameBasedMusicMap.h b/Engine/Plugins/Runtime/MusicEnvironment/Source/MusicEnvironment/Public/FrameBasedMusicMap.h
--- a/Engine/Plugins/Runtime/MusicEnvironment/Source/MusicEnvironment/Public/FrameBasedMusicMap.h
+++ b/Engine/Plugins/Runtime/MusicEnvironment/Source/MusicEnvironment/Public/FrameBasedMusicMap.h
@@ -432,9 +432,144 @@ public:
 	UPROPERTY()
 	FFrameRate FrameResolution;
 
+	/**
+	 * Appends a time signature change to the end of the Bar Map.
+	 * Ensures that this time signature change's tick is after the last one in the map
+	 * 
+	 * @param AtTick
+	 * @param AtBar
+	 * @param Numerator
+	 * @param Denominator
+	 */
 	UE_API void AddTimeSignature(int32 AtTick, int32 AtBar, int16 Numerator, int16 Denominator);
+
+	/**
+	 * Appends a tempo change to the end of the
+	 * Ensures that this tempo change's tick is after the last one in the map
+	 * 
+	 * @param AtTick
+	 * @param AtMs
+	 * @param Bpm
+	 */
 	UE_API void AddTempo(int32 AtTick, float AtMs, float Bpm);
 
+	/**
+	 * Inserts a Time Signature change at the specified Bar, avoiding duplicates as needed.
+	 * Time Signature changes after this point will recalculate their ticks and frame times
+	 * to happen on the new bar lines based on the inserted time signature change
+	 *
+	 * This will not insert a time signature change if the same change already exists immediately before this one being inserted
+	 * If there is the same time signature change immediately after this one being inserted, it will be removed.
+	 *
+	 * @param AtBar
+	 * @param Numerator
+	 * @param Denominator
+	 */
+	UE_API void InsertTimeSignature(int32 AtBar, int16 Numerator, int16 Denominator);
+
+	/**
+	 * Inserts a Tempo Change at the specified tick, avoiding duplicates as needed.
+	 * Time signature and tempo changes after this point will recalculate their frame times and ms respectively
+	 * based on the new tempo change
+	 * 
+	 * This will not insert a tempo change if the same change already exits immediately before
+	 * If the same tempo change exists immediately after this one being inserted, it will be removed
+	 *
+	 * @param AtTick
+	 * @param Bpm
+	 */
+	UE_API void InsertTempo(int32 AtTick, float Bpm);
+
+	/**
+	 * Inserts a Time Signature change at the specified Bar of the Musical Time (ignoring beats), avoiding duplicates as needed.
+	 * Time Signature changes after this point will recalculate their ticks and frame times
+	 * to happen on the new bar lines based on the inserted time signature change
+	 *
+	 * This will not insert a time signature change if the same change already exists immediately before this one being inserted
+	 * If there is the same time signature change immediately after this one being inserted, it will be removed.
+	 *
+	 * @param AtMusicalTime
+	 * @param TimeSignature
+	 */
+	UE_API void InsertTimeSignature(FMusicalTime AtMusicalTime, const FFrameBasedTimeSignature& TimeSignature);
+
+	/**
+	 * Inserts a Tempo Change at the specified MusicalTime, avoiding duplicates as needed.
+	 * Time signature and tempo changes after this point will recalculate their frame times and ms respectively
+	 * based on the new tempo change
+	 *
+	 * This will not insert a tempo change if the same change already exits immediately before
+	 * If the same tempo change exists immediately after this one being inserted, it will be removed
+	 *
+	 * @param AtMusicalTime
+	 * @param Bpm
+	 */
+	UE_API void InsertTempo(FMusicalTime AtMusicalTime, float Bpm);
+
+	/**
+	 * Removes the time signature change at the specified bar, if it exists
+	 * Time Signature changes after this point will recalculate their ticks and frame times
+	 * to happen on the new bar lines based on the inserted time signature change
+	 * 
+	 * @param AtBar
+	 * @return true if a time signature change was removed
+	 */ 
+	UE_API bool RemoveTimeSignature(int32 AtBar);
+
+	/**
+	 * Removes the tempo change at the specified tick, if it exists
+	 * Time signature and tempo changes after this point will recalculate their frame times and ms respectively
+	 * based on the new tempo change
+	 * 
+	 * @param AtTick
+	 * @return true if the tempo change was removed
+	 */
+	UE_API bool RemoveTempo(int32 AtTick);
+
+	/**
+	 * Removes the time signature change at the specified Musical Time, if it exists
+	 *
+	 * Note: Time Signature changes only happen on bar bounderies.
+	 * A Musical Time that is not on whole bare is guaranteed to fail
+	 *
+	 * @param AtMusicalTime
+	 * @return true if a time signature change was removed
+	 */ 	
+	UE_API bool RemoveTimeSignature(FMusicalTime AtMusicalTime);
+	
+	/**
+	 * Removes the tempo change at the specified Musical Time, if it exists
+	 *
+	 * @param AtMusicalTime
+	 * @return true if the tempo change was removed
+	 */
+	UE_API bool RemoveTempo(FMusicalTime AtMusicalTime);
+
+	/**
+	 * @param AtBar
+	 * @return true if there is a time signature change at the given bar
+	 */
+	UE_API bool HasTimeSignatureChange(int32 AtBar);
+
+	/**
+	 * @param AtTick
+	 * @return true if there is a tempo change at the given tick
+	 */
+	UE_API bool HasTempoChange(int32 AtTick);
+
+	/**
+	 * @param AtMusicalTime
+	 * @return true if there is a time signature change at the exact musical time
+	 */
+	UE_API bool HasTimeSignatureChange(FMusicalTime AtMusicalTime);
+
+	/**
+	 * @param AtMusicalTime
+	 * @return true if there is a tempo change at the exact musical time
+	 */
+	UE_API bool HasTempoChange(FMusicalTime AtMusicalTime);
+	
+	
 	UE_API void GatherSubdivisonsForTimespan(int32 StartTick, int32 EndTick, int32 QuantizationInterval, TArray<FGridInterval>& OutSubdivisions) const;
 
 	UE_API void GatherBarsForTimespan(int32 StartBar, int32 EndBar, int32 QuantizationInterval, TArray<FGridInterval>& OutBars) const;
