Fixed up dllstorage using UnrealCodeFixup

diff --git a/Engine/Plugins/Mutable/Source/MutableRuntime/Internal/MuR/OpMeshPrepareLayout.h b/Engine/Plugins/Mutable/Source/MutableRuntime/Internal/MuR/OpMeshPrepareLayout.h
--- a/Engine/Plugins/Mutable/Source/MutableRuntime/Internal/MuR/OpMeshPrepareLayout.h
+++ b/Engine/Plugins/Mutable/Source/MutableRuntime/Internal/MuR/OpMeshPrepareLayout.h
@@ -27,7 +27,7 @@ namespace mu
 
 
 	/** Store the result of a layout vertex classify operation. */
-	struct MUTABLERUNTIME_API FMutableLayoutClassifyResult
+	struct FMutableLayoutClassifyResult
 	{
 		/** The tex coords of every mesh vertex. They could be modified during the classify operation due to forced block clamping. */
 		TArray<FVector2f> TexCoords;
diff --git a/Engine/Plugins/Runtime/MusicEnvironment/Source/MusicEnvironment/Public/MusicTypes/MusicHandle.h b/Engine/Plugins/Runtime/MusicEnvironment/Source/MusicEnvironment/Public/MusicTypes/MusicHandle.h
--- a/Engine/Plugins/Runtime/MusicEnvironment/Source/MusicEnvironment/Public/MusicTypes/MusicHandle.h
+++ b/Engine/Plugins/Runtime/MusicEnvironment/Source/MusicEnvironment/Public/MusicTypes/MusicHandle.h
@@ -58,8 +58,8 @@ enum class EMusicHanldeTransportState : uint8
  * UMusicHandle: The UInterface derived class needed to support Unreal's interface system. This 
  *               class definition stays empty.
  */
-UINTERFACE(BlueprintType, meta = (CannotImplementInterfaceInBlueprint))
-class UE_API UMusicHandle : public UInterface
+UINTERFACE(MinimalAPI, BlueprintType, meta = (CannotImplementInterfaceInBlueprint))
+class UMusicHandle : public UInterface
 {
 	GENERATED_BODY()
 };
@@ -72,7 +72,7 @@ class UE_API UMusicHandle : public UInterface
  *               itself there may not be anything holding a strong reference to the object the
  *               script interface is pointing to. More on this below.
  */
-class UE_API IMusicHandle
+class IMusicHandle
 {
 public:
 	GENERATED_BODY()
@@ -92,9 +92,9 @@ public:
 	UFUNCTION(BlueprintCallable, Category = "Music")
 	virtual void Continue() = 0;
 	UFUNCTION(BlueprintCallable, Category = "Music")
-	virtual void Stop();
+	UE_API virtual void Stop();
 	UFUNCTION(BlueprintCallable, Category = "Music")
-	virtual void Kill();
+	UE_API virtual void Kill();
 
 	operator bool()
 	{
@@ -102,7 +102,7 @@ public:
 	}
 
 	UFUNCTION(BlueprintCallable, Category = "Music", Meta = (ExpandEnumAsExecs = "Branches"))
-	virtual void BranchOnTransportState(EMusicHanldeTransportState& Branches) const;
+	UE_API virtual void BranchOnTransportState(EMusicHanldeTransportState& Branches) const;
 
 	UFUNCTION(BlueprintCallable, Category = "Music")
 	virtual EMusicHanldeTransportState GetTransportState() const { return CurrentTransportState; }
@@ -115,19 +115,19 @@ public:
 	virtual TScriptInterface<IMusicEnvironmentClockSource> GetMusicClockSource() = 0;
 
 	UFUNCTION(BlueprintCallable, Category = "Music|Clock", Meta = (ExpandEnumAsExecs = "Branches"))
-	virtual void GetCurrentBarBeat(float& Bar, float& BeatInBar, EMusicHanldeClockValidity& Branches);
+	UE_API virtual void GetCurrentBarBeat(float& Bar, float& BeatInBar, EMusicHanldeClockValidity& Branches);
 
 	UFUNCTION(BlueprintCallable, Category = "Music|Clock")
-	virtual void BecomeGlobalMusicClockAuthority();
+	UE_API virtual void BecomeGlobalMusicClockAuthority();
 
 	UFUNCTION(BlueprintCallable, Category = "Music|Clock")
-	virtual void RelinquishGlobalMusicClockAuthority();
+	UE_API virtual void RelinquishGlobalMusicClockAuthority();
 	
 	UFUNCTION(BlueprintCallable, Category = "Music|Clock")
-	virtual void RegisterAsTaggedClock(const FGameplayTag& Tag);
+	UE_API virtual void RegisterAsTaggedClock(const FGameplayTag& Tag);
 
 	UFUNCTION(BlueprintCallable, Category = "Music|Clock")
-	virtual void UnregisterAsTaggedClock(const FGameplayTag& Tag = FGameplayTag());
+	UE_API virtual void UnregisterAsTaggedClock(const FGameplayTag& Tag = FGameplayTag());
 
 protected:
 	virtual void Stop_Internal() = 0;
@@ -149,7 +149,7 @@ using FMusicHandle = TScriptInterface<IMusicHandle>;
  *               hold music handles, and for the UStuff inside of those handles to not be garbage collected
  *               until the FStrongMusicHandle is destroyed.
  */
-struct UE_API FStrongMusicHandle
+struct FStrongMusicHandle
 {
 public:
 	FORCEINLINE FStrongMusicHandle() = default;
diff --git a/Engine/Plugins/Runtime/MusicEnvironment/Source/MusicEnvironment/Public/MusicTypes/MusicalAsset.h b/Engine/Plugins/Runtime/MusicEnvironment/Source/MusicEnvironment/Public/MusicTypes/MusicalAsset.h
--- a/Engine/Plugins/Runtime/MusicEnvironment/Source/MusicEnvironment/Public/MusicTypes/MusicalAsset.h
+++ b/Engine/Plugins/Runtime/MusicEnvironment/Source/MusicEnvironment/Public/MusicTypes/MusicalAsset.h
@@ -40,8 +40,8 @@ class IMusicHandle;
  * UMusicalAsset: The UInterface derived class needed to support Unreal's interface system. This
  *                class definition stays empty.
  */
-UINTERFACE(BlueprintType, meta = (CannotImplementInterfaceInBlueprint))
-class UE_API UMusicalAsset : public UMusicMapSource
+UINTERFACE(MinimalAPI, BlueprintType, meta = (CannotImplementInterfaceInBlueprint))
+class UMusicalAsset : public UMusicMapSource
 {
 	GENERATED_BODY()
 };
@@ -52,13 +52,13 @@ class UE_API UMusicalAsset : public UMusicMapSource
  *                FrameBasedMusicMap and can be played as music (returning a 
  *                'music handle' that becomes owned by the caller.)
  */
-class UE_API IMusicalAsset : public IMusicMapSource
+class IMusicalAsset : public IMusicMapSource
 {
 	GENERATED_BODY()
 
 public:
 	UFUNCTION(BlueprintCallable, Category = "Music", Meta = (DefaultToSelf = "PlaybackContext"))
-	virtual TScriptInterface<IMusicHandle> PrepareToPlay(UObject* PlaybackContext, 
+	UE_API virtual TScriptInterface<IMusicHandle> PrepareToPlay(UObject* PlaybackContext, 
 		UAudioComponent* OnComponent = nullptr,
 		float FromSeconds = 0.0f,
 		bool BecomeAuthoritativeClock = false,
@@ -66,7 +66,7 @@ public:
 		bool IsAudition = false);
 
 	UFUNCTION(BlueprintCallable, Category = "Music", Meta = (DefaultToSelf = "PlaybackContext"))
-	virtual TScriptInterface<IMusicHandle> Play(UObject* PlaybackContext,
+	UE_API virtual TScriptInterface<IMusicHandle> Play(UObject* PlaybackContext,
 		UAudioComponent* OnComponent = nullptr,
 		float FromSeconds = 0.0f,
 		bool BecomeAuthoritativeClock = false,
